NNE Rename NNECore module to NNE following split of runtimes into there own plugins

What:
* Following the split of NNE into plugins it became apparent NNECore would be better renamed to NNE to follow UE coding standard, this CL is the first step renaming the module, next step is the related namespace.

Tests:
* All nne tests on win64 editor + style transfer demo

#rnx

diff --git a/Engine/Plugins/AI/MLAdapter/Source/MLAdapter/MLAdapter.Build.cs b/Engine/Plugins/AI/MLAdapter/Source/MLAdapter/MLAdapter.Build.cs
--- a/Engine/Plugins/AI/MLAdapter/Source/MLAdapter/MLAdapter.Build.cs
+++ b/Engine/Plugins/AI/MLAdapter/Source/MLAdapter/MLAdapter.Build.cs
@@ -22,7 +22,7 @@ namespace UnrealBuildTool.Rules
                     "Json",
                     "JsonUtilities",
                     "GameplayAbilities",
-					"NNECore"
+					"NNE"
 				}
             );
 
diff --git a/Engine/Plugins/AI/MLAdapter/Source/MLAdapter/Private/Agents/MLAdapterAgent_Inference.cpp b/Engine/Plugins/AI/MLAdapter/Source/MLAdapter/Private/Agents/MLAdapterAgent_Inference.cpp
--- a/Engine/Plugins/AI/MLAdapter/Source/MLAdapter/Private/Agents/MLAdapterAgent_Inference.cpp
+++ b/Engine/Plugins/AI/MLAdapter/Source/MLAdapter/Private/Agents/MLAdapterAgent_Inference.cpp
@@ -3,8 +3,8 @@
 #include "Agents/MLAdapterAgent_Inference.h"
 #include "HAL/UnrealMemory.h"
 #include "MLAdapterTypes.h"
-#include "NNECore.h"
-#include "NNECoreTypes.h"
+#include "NNE.h"
+#include "NNETypes.h"
 
 void UMLAdapterAgent_Inference::PostInitProperties()
 {
diff --git a/Engine/Plugins/AI/MLAdapter/Source/MLAdapter/Public/Agents/MLAdapterAgent_Inference.h b/Engine/Plugins/AI/MLAdapter/Source/MLAdapter/Public/Agents/MLAdapterAgent_Inference.h
--- a/Engine/Plugins/AI/MLAdapter/Source/MLAdapter/Public/Agents/MLAdapterAgent_Inference.h
+++ b/Engine/Plugins/AI/MLAdapter/Source/MLAdapter/Public/Agents/MLAdapterAgent_Inference.h
@@ -5,8 +5,8 @@
 #include "Templates/UniquePtr.h"
 #include "UObject/SoftObjectPath.h"
 #include "Agents/MLAdapterAgent.h"
-#include "NNECoreModelData.h"
-#include "NNECoreRuntimeCPU.h"
+#include "NNEModelData.h"
+#include "NNERuntimeCPU.h"
 #include "MLAdapterAgent_Inference.generated.h"
 
 /**
diff --git a/Engine/Plugins/Animation/MLDeformer/VertexDeltaModel/Source/VertexDeltaModel/Private/VertexDeltaModel.cpp b/Engine/Plugins/Animation/MLDeformer/VertexDeltaModel/Source/VertexDeltaModel/Private/VertexDeltaModel.cpp
--- a/Engine/Plugins/Animation/MLDeformer/VertexDeltaModel/Source/VertexDeltaModel/Private/VertexDeltaModel.cpp
+++ b/Engine/Plugins/Animation/MLDeformer/VertexDeltaModel/Source/VertexDeltaModel/Private/VertexDeltaModel.cpp
@@ -5,9 +5,9 @@
 #include "VertexDeltaModelInstance.h"
 #include "MLDeformerComponent.h"
 #include "MLDeformerAsset.h"
-#include "NNECore.h"
-#include "NNECoreModelData.h"
-#include "NNECoreRuntimeRDG.h"
+#include "NNE.h"
+#include "NNEModelData.h"
+#include "NNERuntimeRDG.h"
 #include "Modules/ModuleManager.h"
 #include "Interfaces/IPluginManager.h"
 #include "Misc/Paths.h"
diff --git a/Engine/Plugins/Animation/MLDeformer/VertexDeltaModel/Source/VertexDeltaModel/Public/VertexDeltaGraphDataInterface.h b/Engine/Plugins/Animation/MLDeformer/VertexDeltaModel/Source/VertexDeltaModel/Public/VertexDeltaGraphDataInterface.h
--- a/Engine/Plugins/Animation/MLDeformer/VertexDeltaModel/Source/VertexDeltaModel/Public/VertexDeltaGraphDataInterface.h
+++ b/Engine/Plugins/Animation/MLDeformer/VertexDeltaModel/Source/VertexDeltaModel/Public/VertexDeltaGraphDataInterface.h
@@ -8,9 +8,9 @@
 #include "Rendering/SkeletalMeshRenderData.h"
 #include "Rendering/SkeletalMeshLODRenderData.h"
 #include "SkeletalRenderPublic.h"
-#include "NNECore.h"
-#include "NNECoreModelData.h"
-#include "NNECoreRuntimeRDG.h"
+#include "NNE.h"
+#include "NNEModelData.h"
+#include "NNERuntimeRDG.h"
 #include "VertexDeltaGraphDataInterface.generated.h"
 
 class FRDGBuffer;
diff --git a/Engine/Plugins/Animation/MLDeformer/VertexDeltaModel/Source/VertexDeltaModel/Public/VertexDeltaModel.h b/Engine/Plugins/Animation/MLDeformer/VertexDeltaModel/Source/VertexDeltaModel/Public/VertexDeltaModel.h
--- a/Engine/Plugins/Animation/MLDeformer/VertexDeltaModel/Source/VertexDeltaModel/Public/VertexDeltaModel.h
+++ b/Engine/Plugins/Animation/MLDeformer/VertexDeltaModel/Source/VertexDeltaModel/Public/VertexDeltaModel.h
@@ -5,7 +5,7 @@
 #include "CoreTypes.h"
 #include "MLDeformerGeomCacheModel.h"
 #include "VertexDeltaModelVizSettings.h"
-#include "NNECoreRuntimeRDG.h"
+#include "NNERuntimeRDG.h"
 #include "VertexDeltaModel.generated.h"
 
 VERTEXDELTAMODEL_API DECLARE_LOG_CATEGORY_EXTERN(LogVertexDeltaModel, Log, All);
diff --git a/Engine/Plugins/Animation/MLDeformer/VertexDeltaModel/Source/VertexDeltaModel/Public/VertexDeltaModelInstance.h b/Engine/Plugins/Animation/MLDeformer/VertexDeltaModel/Source/VertexDeltaModel/Public/VertexDeltaModelInstance.h
--- a/Engine/Plugins/Animation/MLDeformer/VertexDeltaModel/Source/VertexDeltaModel/Public/VertexDeltaModelInstance.h
+++ b/Engine/Plugins/Animation/MLDeformer/VertexDeltaModel/Source/VertexDeltaModel/Public/VertexDeltaModelInstance.h
@@ -3,10 +3,10 @@
 #pragma once
 
 #include "MLDeformerModelInstance.h"
-#include "NNECore.h"
-#include "NNECoreRuntime.h"
-#include "NNECoreModelData.h"
-#include "NNECoreRuntimeRDG.h"
+#include "NNE.h"
+#include "NNERuntime.h"
+#include "NNEModelData.h"
+#include "NNERuntimeRDG.h"
 #include "RenderGraphResources.h"
 #include "VertexDeltaModelInstance.generated.h"
 
diff --git a/Engine/Plugins/Animation/MLDeformer/VertexDeltaModel/Source/VertexDeltaModel/VertexDeltaModel.Build.cs b/Engine/Plugins/Animation/MLDeformer/VertexDeltaModel/Source/VertexDeltaModel/VertexDeltaModel.Build.cs
--- a/Engine/Plugins/Animation/MLDeformer/VertexDeltaModel/Source/VertexDeltaModel/VertexDeltaModel.Build.cs
+++ b/Engine/Plugins/Animation/MLDeformer/VertexDeltaModel/Source/VertexDeltaModel/VertexDeltaModel.Build.cs
@@ -38,7 +38,7 @@ public class VertexDeltaModel : ModuleRules
 				"CoreUObject",
 				"Engine",
 				"GeometryCache",
-				"NNECore",
+				"NNE",
 				"NNERuntimeRDG",
 				"NNEUtils",
 				"OptimusCore",
diff --git a/Engine/Plugins/Animation/MLDeformer/VertexDeltaModel/Source/VertexDeltaModelEditor/Private/VertexDeltaEditorModel.cpp b/Engine/Plugins/Animation/MLDeformer/VertexDeltaModel/Source/VertexDeltaModelEditor/Private/VertexDeltaEditorModel.cpp
--- a/Engine/Plugins/Animation/MLDeformer/VertexDeltaModel/Source/VertexDeltaModelEditor/Private/VertexDeltaEditorModel.cpp
+++ b/Engine/Plugins/Animation/MLDeformer/VertexDeltaModel/Source/VertexDeltaModelEditor/Private/VertexDeltaEditorModel.cpp
@@ -5,8 +5,8 @@
 #include "VertexDeltaModelVizSettings.h"
 #include "VertexDeltaTrainingModel.h"
 #include "MLDeformerAsset.h"
-#include "NNECore.h"
-#include "NNECoreModelData.h"
+#include "NNE.h"
+#include "NNEModelData.h"
 #include "Misc/Paths.h"
 #include "Misc/FileHelper.h"
 
diff --git a/Engine/Plugins/Animation/MLDeformer/VertexDeltaModel/Source/VertexDeltaModelEditor/VertexDeltaModelEditor.Build.cs b/Engine/Plugins/Animation/MLDeformer/VertexDeltaModel/Source/VertexDeltaModelEditor/VertexDeltaModelEditor.Build.cs
--- a/Engine/Plugins/Animation/MLDeformer/VertexDeltaModel/Source/VertexDeltaModelEditor/VertexDeltaModelEditor.Build.cs
+++ b/Engine/Plugins/Animation/MLDeformer/VertexDeltaModel/Source/VertexDeltaModelEditor/VertexDeltaModelEditor.Build.cs
@@ -31,7 +31,7 @@ namespace UnrealBuildTool.Rules
 					"MLDeformerFrameworkEditor",
 					"VertexDeltaModel",
 					"PropertyEditor",
-					"NNECore",
+					"NNE",
 					"NNERuntimeRDG",
 					"NNEUtils",
 					"ToolWidgets",
diff --git a/Engine/Plugins/Experimental/NNE/NNE.uplugin b/Engine/Plugins/Experimental/NNE/NNE.uplugin
--- a/Engine/Plugins/Experimental/NNE/NNE.uplugin
+++ b/Engine/Plugins/Experimental/NNE/NNE.uplugin
@@ -20,7 +20,7 @@
 	],
 	"Modules": [
 		{
-			"Name": "NNECore",
+			"Name": "NNE",
 			"Type": "RuntimeAndProgram",
 			"LoadingPhase": "Default",
 			"PlatformAllowList": [ "Win64", "Linux", "Mac" ]
diff --git a/Engine/Plugins/Experimental/NNE/Source/NNE/Internal/NNEModelBase.h b/Engine/Plugins/Experimental/NNE/Source/NNE/Internal/NNEModelBase.h
new file mode 100644
--- /dev/null
+++ b/Engine/Plugins/Experimental/NNE/Source/NNE/Internal/NNEModelBase.h
@@ -0,0 +1,80 @@
+// Copyright Epic Games, Inc. All Rights Reserved.
+
+#pragma once
+
+#include "CoreMinimal.h"
+#include "NNE.h"
+#include "NNETypes.h"
+
+namespace UE::NNECore::Internal
+{
+	template <class ModelInterface> class FModelInstanceBase : public ModelInterface
+	{
+	public:
+
+		virtual ~FModelInstanceBase() = default;
+
+		virtual TConstArrayView<NNECore::FTensorDesc> GetInputTensorDescs() const override;
+		virtual TConstArrayView<NNECore::FTensorDesc> GetOutputTensorDescs() const override;
+		virtual TConstArrayView<NNECore::FTensorShape> GetInputTensorShapes() const override;
+		virtual TConstArrayView<NNECore::FTensorShape> GetOutputTensorShapes() const override;
+		virtual int SetInputTensorShapes(TConstArrayView<NNECore::FTensorShape> InInputShapes) override;
+
+	protected:
+
+		FModelInstanceBase() {}
+
+		TArray<FTensorShape>	InputTensorShapes;
+		TArray<FTensorShape>	OutputTensorShapes;
+		TArray<FTensorDesc>		InputSymbolicTensors;
+		TArray<FTensorDesc>		OutputSymbolicTensors;
+	};
+
+	template <class T> TConstArrayView<FTensorDesc> FModelInstanceBase<T>::GetInputTensorDescs() const
+	{
+		return InputSymbolicTensors;
+	}
+
+	template <class T> TConstArrayView<FTensorDesc> FModelInstanceBase<T>::GetOutputTensorDescs() const
+	{
+		return OutputSymbolicTensors;
+	}
+
+	template <class T> TConstArrayView<FTensorShape> FModelInstanceBase<T>::GetInputTensorShapes() const
+	{
+		return InputTensorShapes;
+	}
+
+	template <class T> TConstArrayView<FTensorShape> FModelInstanceBase<T>::GetOutputTensorShapes() const
+	{
+		return OutputTensorShapes;
+	}
+
+	template <class T> int FModelInstanceBase<T>::SetInputTensorShapes(TConstArrayView<FTensorShape> InInputShapes)
+	{
+		InputTensorShapes.Reset(InInputShapes.Num());
+
+		if (InInputShapes.Num() != InputSymbolicTensors.Num())
+		{
+			UE_LOG(LogNNE, Warning, TEXT("Number of input shapes does not match number of input tensors"));
+			return -1;
+		}
+
+		for (int32 i = 0; i < InInputShapes.Num(); ++i)
+		{
+			const FTensorDesc SymbolicDesc = InputSymbolicTensors[i];
+			if (!InInputShapes[i].IsCompatibleWith(SymbolicDesc.GetShape()))
+			{
+				UE_LOG(LogNNE, Warning, TEXT("Input shape does not match input tensor %s of index %d"), *SymbolicDesc.GetName(), i);
+				return -1;
+			}
+		}
+
+		InputTensorShapes = InInputShapes;
+
+		//Implementations are responsible to handle output and intermediate tensor shape inference.
+		//This base implementation only validate that all inputs are matching what the model can support.
+		return 0;
+	}
+
+} // namespace UE::NNECore::Internal
\ No newline at end of file
diff --git a/Engine/Plugins/Experimental/NNE/Source/NNE/Private/NNEAttributeDataType.cpp b/Engine/Plugins/Experimental/NNE/Source/NNE/Private/NNEAttributeDataType.cpp
new file mode 100644
--- /dev/null
+++ b/Engine/Plugins/Experimental/NNE/Source/NNE/Private/NNEAttributeDataType.cpp
@@ -0,0 +1,15 @@
+// Copyright Epic Games, Inc. All Rights Reserved.
+
+#include "NNEAttributeDataType.h"
+
+void LexFromString(ENNEAttributeDataType& OutValue, const TCHAR* StringVal)
+{
+	int64 EnumVal = StaticEnum<ENNEAttributeDataType>()->GetValueByName(StringVal);
+	if (EnumVal == INDEX_NONE)
+	{
+		OutValue = ENNEAttributeDataType::None;
+		ensureMsgf(false, TEXT("ENNEAttributeDataType LexFromString didn't have a match for '%s'"), StringVal);
+	}
+
+	OutValue = (ENNEAttributeDataType)EnumVal;
+};
\ No newline at end of file
diff --git a/Engine/Plugins/Experimental/NNE/Source/NNE/Private/NNEModule.cpp b/Engine/Plugins/Experimental/NNE/Source/NNE/Private/NNEModule.cpp
new file mode 100644
--- /dev/null
+++ b/Engine/Plugins/Experimental/NNE/Source/NNE/Private/NNEModule.cpp
@@ -0,0 +1,19 @@
+// Copyright Epic Games, Inc. All Rights Reserved.
+
+#include "Modules/ModuleManager.h"
+
+class FNNECoreModule : public IModuleInterface
+{
+public:
+
+	virtual void StartupModule() override
+	{
+	}
+
+	virtual void ShutdownModule() override
+	{
+	}
+
+};
+
+IMPLEMENT_MODULE(FNNECoreModule, NNECore)
diff --git a/Engine/Plugins/Experimental/NNE/Source/NNEEditorTools/NNEEditorTools.Build.cs b/Engine/Plugins/Experimental/NNE/Source/NNEEditorTools/NNEEditorTools.Build.cs
--- a/Engine/Plugins/Experimental/NNE/Source/NNEEditorTools/NNEEditorTools.Build.cs
+++ b/Engine/Plugins/Experimental/NNE/Source/NNEEditorTools/NNEEditorTools.Build.cs
@@ -14,7 +14,7 @@ public class NNEEditorTools : ModuleRules
 				"Core",
 				"CoreUObject",
 				"Engine",
-				"NNECore"
+				"NNE"
 			}
 		);
 
diff --git a/Engine/Plugins/Experimental/NNE/Source/NNEEditorTools/Private/NNEEditorToolsModelDataActions.cpp b/Engine/Plugins/Experimental/NNE/Source/NNEEditorTools/Private/NNEEditorToolsModelDataActions.cpp
--- a/Engine/Plugins/Experimental/NNE/Source/NNEEditorTools/Private/NNEEditorToolsModelDataActions.cpp
+++ b/Engine/Plugins/Experimental/NNE/Source/NNEEditorTools/Private/NNEEditorToolsModelDataActions.cpp
@@ -2,7 +2,7 @@
 
 #include "NNEEditorToolsModelDataActions.h"
 #include "EditorFramework/AssetImportData.h"
-#include "NNECoreModelData.h"
+#include "NNEModelData.h"
 #include "NNEEditorToolsModelDataEditorToolkit.h"
 
 #define LOCTEXT_NAMESPACE "AssetTypeActions"
diff --git a/Engine/Plugins/Experimental/NNE/Source/NNEEditorTools/Private/NNEEditorToolsModelDataEditorToolkit.cpp b/Engine/Plugins/Experimental/NNE/Source/NNEEditorTools/Private/NNEEditorToolsModelDataEditorToolkit.cpp
--- a/Engine/Plugins/Experimental/NNE/Source/NNEEditorTools/Private/NNEEditorToolsModelDataEditorToolkit.cpp
+++ b/Engine/Plugins/Experimental/NNE/Source/NNEEditorTools/Private/NNEEditorToolsModelDataEditorToolkit.cpp
@@ -1,8 +1,8 @@
 // Copyright Epic Games, Inc. All Rights Reserved.
 
 #include "NNEEditorToolsModelDataEditorToolkit.h"
-#include "NNECore.h"
-#include "NNECoreModelData.h"
+#include "NNE.h"
+#include "NNEModelData.h"
 #include "SNNEEditorToolsRuntimesWidget.h"
 #include "Widgets/Docking/SDockTab.h"
 
diff --git a/Engine/Plugins/Experimental/NNE/Source/NNEEditorTools/Private/SNNEEditorToolsRuntimesWidget.cpp b/Engine/Plugins/Experimental/NNE/Source/NNEEditorTools/Private/SNNEEditorToolsRuntimesWidget.cpp
--- a/Engine/Plugins/Experimental/NNE/Source/NNEEditorTools/Private/SNNEEditorToolsRuntimesWidget.cpp
+++ b/Engine/Plugins/Experimental/NNE/Source/NNEEditorTools/Private/SNNEEditorToolsRuntimesWidget.cpp
@@ -1,7 +1,7 @@
 // Copyright Epic Games, Inc. All Rights Reserved.
 
 #include "SNNEEditorToolsRuntimesWidget.h"
-#include "NNECore.h"
+#include "NNE.h"
 #include "SlateOptMacros.h"
 #include "Styling/AppStyle.h"
 #include "Widgets/Input/SCheckBox.h"
diff --git a/Engine/Plugins/Experimental/NNERuntimeORTCpu/Source/NNERuntimeORTCpu/Private/NNERuntimeORTCpu.cpp b/Engine/Plugins/Experimental/NNERuntimeORTCpu/Source/NNERuntimeORTCpu/Private/NNERuntimeORTCpu.cpp
--- a/Engine/Plugins/Experimental/NNERuntimeORTCpu/Source/NNERuntimeORTCpu/Private/NNERuntimeORTCpu.cpp
+++ b/Engine/Plugins/Experimental/NNERuntimeORTCpu/Source/NNERuntimeORTCpu/Private/NNERuntimeORTCpu.cpp
@@ -2,9 +2,9 @@
 
 #include "NNERuntimeORTCpu.h"
 
-#include "NNECoreAttributeMap.h"
-#include "NNECoreModelData.h"
-#include "NNECoreModelOptimizerInterface.h"
+#include "NNEAttributeMap.h"
+#include "NNEModelData.h"
+#include "NNEModelOptimizerInterface.h"
 #include "NNEProfilingTimer.h"
 #include "NNERuntimeORTCpuModel.h"
 #include "NNERuntimeORTCpuUtils.h"
diff --git a/Engine/Plugins/Experimental/NNERuntimeORTCpu/Source/NNERuntimeORTCpu/Private/NNERuntimeORTCpu.h b/Engine/Plugins/Experimental/NNERuntimeORTCpu/Source/NNERuntimeORTCpu/Private/NNERuntimeORTCpu.h
--- a/Engine/Plugins/Experimental/NNERuntimeORTCpu/Source/NNERuntimeORTCpu/Private/NNERuntimeORTCpu.h
+++ b/Engine/Plugins/Experimental/NNERuntimeORTCpu/Source/NNERuntimeORTCpu/Private/NNERuntimeORTCpu.h
@@ -3,8 +3,8 @@
 #pragma once
 
 #include "CoreMinimal.h"
-#include "NNECoreRuntime.h"
-#include "NNECoreRuntimeCPU.h"
+#include "NNERuntime.h"
+#include "NNERuntimeCPU.h"
 #include "UObject/Class.h"
 #include "UObject/Object.h"
 #include "UObject/UObjectBaseUtility.h"
diff --git a/Engine/Plugins/Experimental/NNERuntimeORTCpu/Source/NNERuntimeORTCpu/Private/NNERuntimeORTCpuModel.h b/Engine/Plugins/Experimental/NNERuntimeORTCpu/Source/NNERuntimeORTCpu/Private/NNERuntimeORTCpuModel.h
--- a/Engine/Plugins/Experimental/NNERuntimeORTCpu/Source/NNERuntimeORTCpu/Private/NNERuntimeORTCpuModel.h
+++ b/Engine/Plugins/Experimental/NNERuntimeORTCpu/Source/NNERuntimeORTCpu/Private/NNERuntimeORTCpuModel.h
@@ -3,10 +3,10 @@
 #pragma once
 
 #include "CoreMinimal.h"
-#include "NNECoreModelBase.h"
-#include "NNECoreRuntimeCPU.h"
-#include "NNECoreTensor.h"
-#include "NNECoreTypes.h"
+#include "NNEModelBase.h"
+#include "NNERuntimeCPU.h"
+#include "NNETensor.h"
+#include "NNETypes.h"
 #include "NNEProfilingStatistics.h"
 
 #include "NNEThirdPartyWarningDisabler.h"
diff --git a/Engine/Plugins/Experimental/NNERuntimeORTCpu/Source/NNERuntimeORTCpu/Private/NNERuntimeORTCpuModule.cpp b/Engine/Plugins/Experimental/NNERuntimeORTCpu/Source/NNERuntimeORTCpu/Private/NNERuntimeORTCpuModule.cpp
--- a/Engine/Plugins/Experimental/NNERuntimeORTCpu/Source/NNERuntimeORTCpu/Private/NNERuntimeORTCpuModule.cpp
+++ b/Engine/Plugins/Experimental/NNERuntimeORTCpu/Source/NNERuntimeORTCpu/Private/NNERuntimeORTCpuModule.cpp
@@ -1,7 +1,7 @@
 // Copyright Epic Games, Inc. All Rights Reserved.
 
 #include "NNERuntimeORTCpuModule.h"
-#include "NNECore.h"
+#include "NNE.h"
 #include "NNERuntimeORTCpu.h"
 #include "UObject/WeakInterfacePtr.h"
 
diff --git a/Engine/Plugins/Experimental/NNERuntimeORTCpu/Source/NNERuntimeORTCpu/Private/NNERuntimeORTCpuUtils.h b/Engine/Plugins/Experimental/NNERuntimeORTCpu/Source/NNERuntimeORTCpu/Private/NNERuntimeORTCpuUtils.h
--- a/Engine/Plugins/Experimental/NNERuntimeORTCpu/Source/NNERuntimeORTCpu/Private/NNERuntimeORTCpuUtils.h
+++ b/Engine/Plugins/Experimental/NNERuntimeORTCpu/Source/NNERuntimeORTCpu/Private/NNERuntimeORTCpuUtils.h
@@ -3,10 +3,10 @@
 #pragma once
 
 #include "HAL/UnrealMemory.h"
-#include "NNECore.h"
-#include "NNECoreRuntimeCPU.h"
-#include "NNECoreTypes.h"
-#include "NNECoreTensor.h"
+#include "NNE.h"
+#include "NNERuntimeCPU.h"
+#include "NNETypes.h"
+#include "NNETensor.h"
 #include "NNEThirdPartyWarningDisabler.h"
 NNE_THIRD_PARTY_INCLUDES_START
 #undef check
diff --git a/Engine/Plugins/Experimental/NNERuntimeORTCpu/Source/NNEUtils/Internal/NNEUtilsLogHelper.h b/Engine/Plugins/Experimental/NNERuntimeORTCpu/Source/NNEUtils/Internal/NNEUtilsLogHelper.h
--- a/Engine/Plugins/Experimental/NNERuntimeORTCpu/Source/NNEUtils/Internal/NNEUtilsLogHelper.h
+++ b/Engine/Plugins/Experimental/NNERuntimeORTCpu/Source/NNEUtils/Internal/NNEUtilsLogHelper.h
@@ -1,7 +1,7 @@
 // Copyright Epic Games, Inc. All Rights Reserved.
 #pragma once
 
-#include "NNECoreTypes.h"
+#include "NNETypes.h"
 
 namespace UE::NNEUtils::Internal
 {
diff --git a/Engine/Plugins/Experimental/NNERuntimeORTCpu/Source/NNEUtils/Internal/NNEUtilsModelBuilder.h b/Engine/Plugins/Experimental/NNERuntimeORTCpu/Source/NNEUtils/Internal/NNEUtilsModelBuilder.h
--- a/Engine/Plugins/Experimental/NNERuntimeORTCpu/Source/NNEUtils/Internal/NNEUtilsModelBuilder.h
+++ b/Engine/Plugins/Experimental/NNERuntimeORTCpu/Source/NNEUtils/Internal/NNEUtilsModelBuilder.h
@@ -1,8 +1,8 @@
 // Copyright Epic Games, Inc. All Rights Reserved.
 #pragma once
 
-#include "NNECoreTypes.h"
-#include "NNECoreTensor.h"
+#include "NNETypes.h"
+#include "NNETensor.h"
 
 namespace UE::NNECore { class FAttributeMap; }
 struct FNNEAttributeValue;
diff --git a/Engine/Plugins/Experimental/NNERuntimeORTCpu/Source/NNEUtils/Internal/NNEUtilsModelOptimizer.h b/Engine/Plugins/Experimental/NNERuntimeORTCpu/Source/NNEUtils/Internal/NNEUtilsModelOptimizer.h
--- a/Engine/Plugins/Experimental/NNERuntimeORTCpu/Source/NNEUtils/Internal/NNEUtilsModelOptimizer.h
+++ b/Engine/Plugins/Experimental/NNERuntimeORTCpu/Source/NNEUtils/Internal/NNEUtilsModelOptimizer.h
@@ -1,7 +1,7 @@
 // Copyright Epic Games, Inc. All Rights Reserved.
 #pragma once
 
-#include "NNECoreRuntimeFormat.h"
+#include "NNERuntimeFormat.h"
 
 namespace UE::NNECore { class FAttributeMap; }
 namespace UE::NNECore::Internal { class IModelOptimizer; }
diff --git a/Engine/Plugins/Experimental/NNERuntimeORTCpu/Source/NNEUtils/NNEUtils.Build.cs b/Engine/Plugins/Experimental/NNERuntimeORTCpu/Source/NNEUtils/NNEUtils.Build.cs
--- a/Engine/Plugins/Experimental/NNERuntimeORTCpu/Source/NNEUtils/NNEUtils.Build.cs
+++ b/Engine/Plugins/Experimental/NNERuntimeORTCpu/Source/NNEUtils/NNEUtils.Build.cs
@@ -13,14 +13,14 @@ public class NNEUtils : ModuleRules
 
 		PublicIncludePaths.AddRange(
 			new string[] {
-				Path.GetFullPath(Path.Combine(EngineDirectory, "Plugins/Experimental/NNE/Source/NNECore/Private")),
+				Path.GetFullPath(Path.Combine(EngineDirectory, "Plugins/Experimental/NNE/Source/NNE/Private")),
 			}
 		);
 
 		PublicDependencyModuleNames.AddRange(
 			new string[]
 			{
-				"NNECore"
+				"NNE"
 			}
 		);
 
diff --git a/Engine/Plugins/Experimental/NNERuntimeORTCpu/Source/NNEUtils/Private/NNEUtilsModelBuilderNNE.cpp b/Engine/Plugins/Experimental/NNERuntimeORTCpu/Source/NNEUtils/Private/NNEUtilsModelBuilderNNE.cpp
--- a/Engine/Plugins/Experimental/NNERuntimeORTCpu/Source/NNEUtils/Private/NNEUtilsModelBuilderNNE.cpp
+++ b/Engine/Plugins/Experimental/NNERuntimeORTCpu/Source/NNEUtils/Private/NNEUtilsModelBuilderNNE.cpp
@@ -2,9 +2,9 @@
 
 #include "NNEUtilsModelBuilder.h"
 
-#include "NNECore.h"
-#include "NNECoreAttributeMap.h"
-#include "NNECoreRuntimeFormat.h"
+#include "NNE.h"
+#include "NNEAttributeMap.h"
+#include "NNERuntimeFormat.h"
 #include "Misc/StringBuilder.h"
 #include "Serialization/MemoryWriter.h"
 
diff --git a/Engine/Plugins/Experimental/NNERuntimeORTCpu/Source/NNEUtils/Private/NNEUtilsModelBuilderONNX.cpp b/Engine/Plugins/Experimental/NNERuntimeORTCpu/Source/NNEUtils/Private/NNEUtilsModelBuilderONNX.cpp
--- a/Engine/Plugins/Experimental/NNERuntimeORTCpu/Source/NNEUtils/Private/NNEUtilsModelBuilderONNX.cpp
+++ b/Engine/Plugins/Experimental/NNERuntimeORTCpu/Source/NNEUtils/Private/NNEUtilsModelBuilderONNX.cpp
@@ -1,9 +1,9 @@
 // Copyright Epic Games, Inc. All Rights Reserved.
 
-#include "NNECore.h"
+#include "NNE.h"
 #include "NNEUtilsModelBuilder.h"
-#include "NNECoreAttributeMap.h"
-#include "NNECoreRuntimeFormat.h"
+#include "NNEAttributeMap.h"
+#include "NNERuntimeFormat.h"
 
 #include "NNEThirdPartyWarningDisabler.h"
 NNE_THIRD_PARTY_INCLUDES_START
diff --git a/Engine/Plugins/Experimental/NNERuntimeORTCpu/Source/NNEUtils/Private/NNEUtilsModelOptimizer.cpp b/Engine/Plugins/Experimental/NNERuntimeORTCpu/Source/NNEUtils/Private/NNEUtilsModelOptimizer.cpp
--- a/Engine/Plugins/Experimental/NNERuntimeORTCpu/Source/NNEUtils/Private/NNEUtilsModelOptimizer.cpp
+++ b/Engine/Plugins/Experimental/NNERuntimeORTCpu/Source/NNEUtils/Private/NNEUtilsModelOptimizer.cpp
@@ -2,9 +2,9 @@
 
 #include "NNEUtilsModelOptimizer.h"
 
-#include "NNECore.h"
-#include "NNECoreModelOptimizerInterface.h"
-#include "NNECoreRuntimeFormat.h"
+#include "NNE.h"
+#include "NNEModelOptimizerInterface.h"
+#include "NNERuntimeFormat.h"
 #include "NNEUtilsModelBuilder.h"
 
 #ifdef PLATFORM_NNE_MICROSOFT
diff --git a/Engine/Plugins/Experimental/NNERuntimeORTGpu/Source/NNERuntimeORT/Private/NNERuntimeORT.cpp b/Engine/Plugins/Experimental/NNERuntimeORTGpu/Source/NNERuntimeORT/Private/NNERuntimeORT.cpp
--- a/Engine/Plugins/Experimental/NNERuntimeORTGpu/Source/NNERuntimeORT/Private/NNERuntimeORT.cpp
+++ b/Engine/Plugins/Experimental/NNERuntimeORTGpu/Source/NNERuntimeORT/Private/NNERuntimeORT.cpp
@@ -4,9 +4,9 @@
 #include "NNERuntimeORTModel.h"
 #include "NNERuntimeORTUtils.h"
 #include "NNEUtilsModelOptimizer.h"
-#include "NNECoreAttributeMap.h"
-#include "NNECoreModelData.h"
-#include "NNECoreModelOptimizerInterface.h"
+#include "NNEAttributeMap.h"
+#include "NNEModelData.h"
+#include "NNEModelOptimizerInterface.h"
 #include "NNEProfilingTimer.h"
 
 #include UE_INLINE_GENERATED_CPP_BY_NAME(NNERuntimeORT)
diff --git a/Engine/Plugins/Experimental/NNERuntimeORTGpu/Source/NNERuntimeORT/Private/NNERuntimeORT.h b/Engine/Plugins/Experimental/NNERuntimeORTGpu/Source/NNERuntimeORT/Private/NNERuntimeORT.h
--- a/Engine/Plugins/Experimental/NNERuntimeORTGpu/Source/NNERuntimeORT/Private/NNERuntimeORT.h
+++ b/Engine/Plugins/Experimental/NNERuntimeORTGpu/Source/NNERuntimeORT/Private/NNERuntimeORT.h
@@ -3,8 +3,8 @@
 #pragma once
 
 #include "CoreMinimal.h"
-#include "NNECoreRuntime.h"
-#include "NNECoreRuntimeGPU.h"
+#include "NNERuntime.h"
+#include "NNERuntimeGPU.h"
 #include "Templates/UniquePtr.h"
 #include "UObject/Class.h"
 #include "UObject/Object.h"
diff --git a/Engine/Plugins/Experimental/NNERuntimeORTGpu/Source/NNERuntimeORT/Private/NNERuntimeORTModel.h b/Engine/Plugins/Experimental/NNERuntimeORTGpu/Source/NNERuntimeORT/Private/NNERuntimeORTModel.h
--- a/Engine/Plugins/Experimental/NNERuntimeORTGpu/Source/NNERuntimeORT/Private/NNERuntimeORTModel.h
+++ b/Engine/Plugins/Experimental/NNERuntimeORTGpu/Source/NNERuntimeORT/Private/NNERuntimeORTModel.h
@@ -3,10 +3,10 @@
 #pragma once
 
 #include "CoreMinimal.h"
-#include "NNECoreModelBase.h"
-#include "NNECoreRuntimeGPU.h"
-#include "NNECoreTensor.h"
-#include "NNECoreTypes.h"
+#include "NNEModelBase.h"
+#include "NNERuntimeGPU.h"
+#include "NNETensor.h"
+#include "NNETypes.h"
 #include "NNEProfilingStatistics.h"
 
 #include "NNEThirdPartyWarningDisabler.h"
diff --git a/Engine/Plugins/Experimental/NNERuntimeORTGpu/Source/NNERuntimeORT/Private/NNERuntimeORTUtils.h b/Engine/Plugins/Experimental/NNERuntimeORTGpu/Source/NNERuntimeORT/Private/NNERuntimeORTUtils.h
--- a/Engine/Plugins/Experimental/NNERuntimeORTGpu/Source/NNERuntimeORT/Private/NNERuntimeORTUtils.h
+++ b/Engine/Plugins/Experimental/NNERuntimeORTGpu/Source/NNERuntimeORT/Private/NNERuntimeORTUtils.h
@@ -3,10 +3,10 @@
 #pragma once
 
 #include "HAL/UnrealMemory.h"
-#include "NNECore.h"
-#include "NNECoreRuntimeGPU.h"
-#include "NNECoreTypes.h"
-#include "NNECoreTensor.h"
+#include "NNE.h"
+#include "NNERuntimeGPU.h"
+#include "NNETypes.h"
+#include "NNETensor.h"
 #include "NNEThirdPartyWarningDisabler.h"
 NNE_THIRD_PARTY_INCLUDES_START
 #undef check
diff --git a/Engine/Plugins/Experimental/NNERuntimeRDG/Source/NNEHlslShaders/Internal/NNEHlslShadersBatchNormalizationCS.h b/Engine/Plugins/Experimental/NNERuntimeRDG/Source/NNEHlslShaders/Internal/NNEHlslShadersBatchNormalizationCS.h
--- a/Engine/Plugins/Experimental/NNERuntimeRDG/Source/NNEHlslShaders/Internal/NNEHlslShadersBatchNormalizationCS.h
+++ b/Engine/Plugins/Experimental/NNERuntimeRDG/Source/NNEHlslShaders/Internal/NNEHlslShadersBatchNormalizationCS.h
@@ -2,7 +2,7 @@
 
 #pragma once
 
-#include "NNECoreOperator.h"
+#include "NNEOperator.h"
 #include "NNEHlslShaderBase.h"
 #include "RenderGraphUtils.h"
 #include "ShaderParameterUtils.h"
diff --git a/Engine/Plugins/Experimental/NNERuntimeRDG/Source/NNEHlslShaders/Internal/NNEHlslShadersElementWiseBinaryCS.h b/Engine/Plugins/Experimental/NNERuntimeRDG/Source/NNEHlslShaders/Internal/NNEHlslShadersElementWiseBinaryCS.h
--- a/Engine/Plugins/Experimental/NNERuntimeRDG/Source/NNEHlslShaders/Internal/NNEHlslShadersElementWiseBinaryCS.h
+++ b/Engine/Plugins/Experimental/NNERuntimeRDG/Source/NNEHlslShaders/Internal/NNEHlslShadersElementWiseBinaryCS.h
@@ -2,7 +2,7 @@
 
 #pragma once
 
-#include "NNECoreOperator.h"
+#include "NNEOperator.h"
 #include "NNEHlslShaderBase.h"
 #include "RenderGraphUtils.h"
 #include "ShaderParameterUtils.h"
diff --git a/Engine/Plugins/Experimental/NNERuntimeRDG/Source/NNEHlslShaders/Private/NNEHlslShadersGatherCS.cpp b/Engine/Plugins/Experimental/NNERuntimeRDG/Source/NNEHlslShaders/Private/NNEHlslShadersGatherCS.cpp
--- a/Engine/Plugins/Experimental/NNERuntimeRDG/Source/NNEHlslShaders/Private/NNEHlslShadersGatherCS.cpp
+++ b/Engine/Plugins/Experimental/NNERuntimeRDG/Source/NNEHlslShaders/Private/NNEHlslShadersGatherCS.cpp
@@ -1,7 +1,7 @@
 // Copyright Epic Games, Inc. All Rights Reserved.
 
 #include "NNEHlslShadersGatherCS.h"
-#include "NNECoreTensor.h"
+#include "NNETensor.h"
 
 namespace UE::NNEHlslShaders::Internal
 {
diff --git a/Engine/Plugins/Experimental/NNERuntimeRDG/Source/NNEHlslShaders/Private/NNEHlslShadersGemmCS.cpp b/Engine/Plugins/Experimental/NNERuntimeRDG/Source/NNEHlslShaders/Private/NNEHlslShadersGemmCS.cpp
--- a/Engine/Plugins/Experimental/NNERuntimeRDG/Source/NNEHlslShaders/Private/NNEHlslShadersGemmCS.cpp
+++ b/Engine/Plugins/Experimental/NNERuntimeRDG/Source/NNEHlslShaders/Private/NNEHlslShadersGemmCS.cpp
@@ -1,8 +1,8 @@
 // Copyright Epic Games, Inc. All Rights Reserved.
 
 #include "NNEHlslShadersGemmCS.h"
-#include "NNECore.h"
-#include "NNECoreTensor.h"
+#include "NNE.h"
+#include "NNETensor.h"
 #include "Algo/Accumulate.h"
 
 namespace UE::NNEHlslShaders::Internal
diff --git a/Engine/Plugins/Experimental/NNERuntimeRDG/Source/NNEHlslShaders/Private/NNEHlslShadersInstanceNormalizationCS.cpp b/Engine/Plugins/Experimental/NNERuntimeRDG/Source/NNEHlslShaders/Private/NNEHlslShadersInstanceNormalizationCS.cpp
--- a/Engine/Plugins/Experimental/NNERuntimeRDG/Source/NNEHlslShaders/Private/NNEHlslShadersInstanceNormalizationCS.cpp
+++ b/Engine/Plugins/Experimental/NNERuntimeRDG/Source/NNEHlslShaders/Private/NNEHlslShadersInstanceNormalizationCS.cpp
@@ -1,8 +1,8 @@
 // Copyright Epic Games, Inc. All Rights Reserved.
 
 #include "NNEHlslShadersInstanceNormalizationCS.h"
-#include "NNECore.h"
-#include "NNECoreTensor.h"
+#include "NNE.h"
+#include "NNETensor.h"
 
 namespace UE::NNEHlslShaders::Internal
 {
diff --git a/Engine/Plugins/Experimental/NNERuntimeRDG/Source/NNERuntimeRDG/Internal/NNERuntimeRDGHelperConcat.h b/Engine/Plugins/Experimental/NNERuntimeRDG/Source/NNERuntimeRDG/Internal/NNERuntimeRDGHelperConcat.h
--- a/Engine/Plugins/Experimental/NNERuntimeRDG/Source/NNERuntimeRDG/Internal/NNERuntimeRDGHelperConcat.h
+++ b/Engine/Plugins/Experimental/NNERuntimeRDG/Source/NNERuntimeRDG/Internal/NNERuntimeRDGHelperConcat.h
@@ -3,7 +3,7 @@
 #pragma once 
 
 #include "CoreMinimal.h"
-#include "NNECoreTensor.h"
+#include "NNETensor.h"
 
 namespace UE::NNERuntimeRDG::Internal::CPUHelper::Concat
 {
diff --git a/Engine/Plugins/Experimental/NNERuntimeRDG/Source/NNERuntimeRDG/Private/Dml/NNEDmlOperator.h b/Engine/Plugins/Experimental/NNERuntimeRDG/Source/NNERuntimeRDG/Private/Dml/NNEDmlOperator.h
--- a/Engine/Plugins/Experimental/NNERuntimeRDG/Source/NNERuntimeRDG/Private/Dml/NNEDmlOperator.h
+++ b/Engine/Plugins/Experimental/NNERuntimeRDG/Source/NNERuntimeRDG/Private/Dml/NNEDmlOperator.h
@@ -5,10 +5,10 @@
 #ifdef NNE_USE_DIRECTML
 
 #include "NNEDmlCommon.h"
-#include "NNECoreTypes.h"
-#include "NNECoreTensor.h"
-#include "NNECoreAttributeMap.h"
-#include "NNERuntimeRDG.h"
+#include "NNETypes.h"
+#include "NNETensor.h"
+#include "NNEAttributeMap.h"
+#include "NNERuntimeRDGBase.h"
 
 #define NNE_DML_REGISTER_OP(OpName) \
 struct FDmlOperator##OpName##Registrator \
diff --git a/Engine/Plugins/Experimental/NNERuntimeRDG/Source/NNERuntimeRDG/Private/Helper/NNERuntimeRDGHelperConcat.cpp b/Engine/Plugins/Experimental/NNERuntimeRDG/Source/NNERuntimeRDG/Private/Helper/NNERuntimeRDGHelperConcat.cpp
--- a/Engine/Plugins/Experimental/NNERuntimeRDG/Source/NNERuntimeRDG/Private/Helper/NNERuntimeRDGHelperConcat.cpp
+++ b/Engine/Plugins/Experimental/NNERuntimeRDG/Source/NNERuntimeRDG/Private/Helper/NNERuntimeRDGHelperConcat.cpp
@@ -2,7 +2,7 @@
 
 #include "NNERuntimeRDGHelperConcat.h"
 
-#include "NNECoreTypes.h"
+#include "NNETypes.h"
 #include "NNERuntimeRDGTensorIdxIterator.h"
 
 namespace UE::NNERuntimeRDG::Internal::CPUHelper::Concat
diff --git a/Engine/Plugins/Experimental/NNERuntimeRDG/Source/NNERuntimeRDG/Private/Helper/NNERuntimeRDGHelperElementWiseBinary.cpp b/Engine/Plugins/Experimental/NNERuntimeRDG/Source/NNERuntimeRDG/Private/Helper/NNERuntimeRDGHelperElementWiseBinary.cpp
--- a/Engine/Plugins/Experimental/NNERuntimeRDG/Source/NNERuntimeRDG/Private/Helper/NNERuntimeRDGHelperElementWiseBinary.cpp
+++ b/Engine/Plugins/Experimental/NNERuntimeRDG/Source/NNERuntimeRDG/Private/Helper/NNERuntimeRDGHelperElementWiseBinary.cpp
@@ -2,8 +2,8 @@
 
 #include "NNERuntimeRDGHelperElementWiseBinary.h"
 #include "NNERuntimeRDGTensorIdxIterator.h"
-#include "NNECoreTypes.h"
-#include "NNECoreTensor.h"
+#include "NNETypes.h"
+#include "NNETensor.h"
 #include "Math/UnrealMathUtility.h"
 #include "MathUtil.h"
 
diff --git a/Engine/Plugins/Experimental/NNERuntimeRDG/Source/NNERuntimeRDG/Private/Helper/NNERuntimeRDGHelperReshape.h b/Engine/Plugins/Experimental/NNERuntimeRDG/Source/NNERuntimeRDG/Private/Helper/NNERuntimeRDGHelperReshape.h
--- a/Engine/Plugins/Experimental/NNERuntimeRDG/Source/NNERuntimeRDG/Private/Helper/NNERuntimeRDGHelperReshape.h
+++ b/Engine/Plugins/Experimental/NNERuntimeRDG/Source/NNERuntimeRDG/Private/Helper/NNERuntimeRDGHelperReshape.h
@@ -6,8 +6,8 @@
 #include "Algo/Count.h"
 #include "CoreMinimal.h"
 #include "Misc/EnumerateRange.h"
-#include "NNECoreTensor.h"
-#include "NNECoreTypes.h"
+#include "NNETensor.h"
+#include "NNETypes.h"
 
 namespace UE::NNERuntimeRDG::Private::ShapeHelper::Reshape
 {
diff --git a/Engine/Plugins/Experimental/NNERuntimeRDG/Source/NNERuntimeRDG/Private/Hlsl/NNERuntimeRDGBatchNormalization.cpp b/Engine/Plugins/Experimental/NNERuntimeRDG/Source/NNERuntimeRDG/Private/Hlsl/NNERuntimeRDGBatchNormalization.cpp
--- a/Engine/Plugins/Experimental/NNERuntimeRDG/Source/NNERuntimeRDG/Private/Hlsl/NNERuntimeRDGBatchNormalization.cpp
+++ b/Engine/Plugins/Experimental/NNERuntimeRDG/Source/NNERuntimeRDG/Private/Hlsl/NNERuntimeRDGBatchNormalization.cpp
@@ -3,9 +3,9 @@
 #include "NNERuntimeRDGBatchNormalization.h"
 #include "NNEHlslShadersBatchNormalizationCS.h"
 #include "NNERuntimeRDGHlslHelper.h"
-#include "NNECoreAttributeMap.h"
-#include "NNECoreTypes.h"
-#include "NNECoreTensor.h"
+#include "NNEAttributeMap.h"
+#include "NNETypes.h"
+#include "NNETensor.h"
 
 namespace UE::NNERuntimeRDG::Private::Hlsl
 {
diff --git a/Engine/Plugins/Experimental/NNERuntimeRDG/Source/NNERuntimeRDG/Private/Hlsl/NNERuntimeRDGBatchNormalization.h b/Engine/Plugins/Experimental/NNERuntimeRDG/Source/NNERuntimeRDG/Private/Hlsl/NNERuntimeRDGBatchNormalization.h
--- a/Engine/Plugins/Experimental/NNERuntimeRDG/Source/NNERuntimeRDG/Private/Hlsl/NNERuntimeRDGBatchNormalization.h
+++ b/Engine/Plugins/Experimental/NNERuntimeRDG/Source/NNERuntimeRDG/Private/Hlsl/NNERuntimeRDGBatchNormalization.h
@@ -2,7 +2,7 @@
 
 #pragma once
 
-#include "NNERuntimeRDG.h"
+#include "NNERuntimeRDGBase.h"
 #include "NNERuntimeRDGHlslOp.h"
 
 namespace UE::NNERuntimeRDG::Private::Hlsl
diff --git a/Engine/Plugins/Experimental/NNERuntimeRDG/Source/NNERuntimeRDG/Private/Hlsl/NNERuntimeRDGCast.cpp b/Engine/Plugins/Experimental/NNERuntimeRDG/Source/NNERuntimeRDG/Private/Hlsl/NNERuntimeRDGCast.cpp
--- a/Engine/Plugins/Experimental/NNERuntimeRDG/Source/NNERuntimeRDG/Private/Hlsl/NNERuntimeRDGCast.cpp
+++ b/Engine/Plugins/Experimental/NNERuntimeRDG/Source/NNERuntimeRDG/Private/Hlsl/NNERuntimeRDGCast.cpp
@@ -2,8 +2,8 @@
 
 #include "NNERuntimeRDGCast.h"
 #include "NNERuntimeRDGHelperCast.h"
-#include "NNECoreTensor.h"
-#include "NNECoreTypes.h"
+#include "NNETensor.h"
+#include "NNETypes.h"
 #include "NNEUtilsLogHelper.h"
 
 namespace UE::NNERuntimeRDG::Private::Hlsl
diff --git a/Engine/Plugins/Experimental/NNERuntimeRDG/Source/NNERuntimeRDG/Private/Hlsl/NNERuntimeRDGCast.h b/Engine/Plugins/Experimental/NNERuntimeRDG/Source/NNERuntimeRDG/Private/Hlsl/NNERuntimeRDGCast.h
--- a/Engine/Plugins/Experimental/NNERuntimeRDG/Source/NNERuntimeRDG/Private/Hlsl/NNERuntimeRDGCast.h
+++ b/Engine/Plugins/Experimental/NNERuntimeRDG/Source/NNERuntimeRDG/Private/Hlsl/NNERuntimeRDGCast.h
@@ -2,7 +2,7 @@
 
 #pragma once
 
-#include "NNERuntimeRDG.h"
+#include "NNERuntimeRDGBase.h"
 #include "NNERuntimeRDGHlslOp.h"
 
 namespace UE::NNERuntimeRDG::Private::Hlsl
diff --git a/Engine/Plugins/Experimental/NNERuntimeRDG/Source/NNERuntimeRDG/Private/Hlsl/NNERuntimeRDGConv.cpp b/Engine/Plugins/Experimental/NNERuntimeRDG/Source/NNERuntimeRDG/Private/Hlsl/NNERuntimeRDGConv.cpp
--- a/Engine/Plugins/Experimental/NNERuntimeRDG/Source/NNERuntimeRDG/Private/Hlsl/NNERuntimeRDGConv.cpp
+++ b/Engine/Plugins/Experimental/NNERuntimeRDG/Source/NNERuntimeRDG/Private/Hlsl/NNERuntimeRDGConv.cpp
@@ -2,9 +2,9 @@
 
 #include "NNERuntimeRDGConv.h"
 #include "Algo/AnyOf.h"
-#include "NNECoreAttributeMap.h"
-#include "NNECoreTensor.h"
-#include "NNECoreTypes.h"
+#include "NNEAttributeMap.h"
+#include "NNETensor.h"
+#include "NNETypes.h"
 #include "NNEHlslShadersConvCS.h"
 #include "NNEHlslShadersConvMatmulCS.h"
 #include "NNERuntimeRDGHelperTranspose.h"
diff --git a/Engine/Plugins/Experimental/NNERuntimeRDG/Source/NNERuntimeRDG/Private/Hlsl/NNERuntimeRDGElementWiseUnary.cpp b/Engine/Plugins/Experimental/NNERuntimeRDG/Source/NNERuntimeRDG/Private/Hlsl/NNERuntimeRDGElementWiseUnary.cpp
--- a/Engine/Plugins/Experimental/NNERuntimeRDG/Source/NNERuntimeRDG/Private/Hlsl/NNERuntimeRDGElementWiseUnary.cpp
+++ b/Engine/Plugins/Experimental/NNERuntimeRDG/Source/NNERuntimeRDG/Private/Hlsl/NNERuntimeRDGElementWiseUnary.cpp
@@ -4,9 +4,9 @@
 #include "NNERuntimeRDGHelperElementWiseUnary.h"
 #include "NNEHlslShadersElementWiseUnaryCS.h"
 #include "NNERuntimeRDGHlslHelper.h"
-#include "NNECoreAttributeMap.h"
-#include "NNECoreTypes.h"
-#include "NNECoreTensor.h"
+#include "NNEAttributeMap.h"
+#include "NNETypes.h"
+#include "NNETensor.h"
 
 namespace UE::NNERuntimeRDG::Private::Hlsl
 {
diff --git a/Engine/Plugins/Experimental/NNERuntimeRDG/Source/NNERuntimeRDG/Private/Hlsl/NNERuntimeRDGGemm.cpp b/Engine/Plugins/Experimental/NNERuntimeRDG/Source/NNERuntimeRDG/Private/Hlsl/NNERuntimeRDGGemm.cpp
--- a/Engine/Plugins/Experimental/NNERuntimeRDG/Source/NNERuntimeRDG/Private/Hlsl/NNERuntimeRDGGemm.cpp
+++ b/Engine/Plugins/Experimental/NNERuntimeRDG/Source/NNERuntimeRDG/Private/Hlsl/NNERuntimeRDGGemm.cpp
@@ -3,9 +3,9 @@
 #include "NNERuntimeRDGGemm.h"
 #include "NNEHlslShadersGemmCS.h"
 #include "NNERuntimeRDGHlslHelper.h"
-#include "NNECoreAttributeMap.h"
-#include "NNECoreTensor.h"
-#include "NNECoreTypes.h"
+#include "NNEAttributeMap.h"
+#include "NNETensor.h"
+#include "NNETypes.h"
 
 namespace UE::NNERuntimeRDG::Private::Hlsl
 {
diff --git a/Engine/Plugins/Experimental/NNERuntimeRDG/Source/NNERuntimeRDG/Private/NNERuntimeRDGBase.cpp b/Engine/Plugins/Experimental/NNERuntimeRDG/Source/NNERuntimeRDG/Private/NNERuntimeRDGBase.cpp
new file mode 100644
--- /dev/null
+++ b/Engine/Plugins/Experimental/NNERuntimeRDG/Source/NNERuntimeRDG/Private/NNERuntimeRDGBase.cpp
@@ -0,0 +1,166 @@
+// Copyright Epic Games, Inc. All Rights Reserved.
+
+#include "NNERuntimeRDGBase.h"
+
+#include "NNERuntimeFormat.h"
+#include "NNEUtilsLogHelper.h"
+#include "NNEUtilsModelOptimizer.h"
+#include "RenderGraphBuilder.h"
+#include "RenderGraphUtils.h"
+#include "RHIGPUReadback.h"
+#include "Serialization/MemoryReader.h"
+
+namespace UE::NNERuntimeRDG::Private
+{
+
+//
+//
+//
+bool AlwaysValidValidationFunction(
+	const NNECore::FAttributeMap& AttributeMap, 
+	TConstArrayView<ENNETensorDataType> InputTensorTypes,
+	TConstArrayView<NNECore::FSymbolicTensorShape> InputShapes)
+{
+	return true;
+}
+
+
+//
+//
+//
+FInputValidator::FInputValidator() : 
+	NumRequiredInput(0), NumOptionalInput(0)
+{
+	TemplateTypes.SetNum(1);
+}
+
+bool FInputValidator::Validate(TConstArrayView<ENNETensorDataType> InputTypes)
+{
+	check(InputTemplateIndices.Num() == NumRequiredInput + NumOptionalInput);
+
+	bool bAreInputValid = true;
+	int32 NumInputsToValidate = FMath::Min(InputTemplateIndices.Num(), InputTypes.Num());
+	
+	if (InputTypes.Num() < NumRequiredInput)
+	{
+		UE_LOG(LogNNE, Warning, TEXT("Required '%d' inputs but found '%d'."), NumRequiredInput, InputTypes.Num());
+		bAreInputValid = false;
+	}
+	if (InputTypes.Num() > NumRequiredInput+NumOptionalInput)
+	{
+		UE_LOG(LogNNE, Warning, TEXT("Got a total of '%d' inputs but should have '%d' maximum."), InputTypes.Num(), NumRequiredInput + NumOptionalInput);
+		bAreInputValid = false;
+	}
+	
+	for (int32 Idx = 0; Idx < NumInputsToValidate; ++Idx)
+	{
+		const int32 TemplateIdx = InputTemplateIndices[Idx];
+		
+		check(TemplateIdx < TemplateTypes.Num());
+		if (INDEX_NONE == TemplateTypes[TemplateIdx].Find(InputTypes[Idx]))
+		{
+			FString TargetType = NNEUtils::Internal::GetTensorDataTypeName(InputTypes[Idx]);
+			UE_LOG(LogNNE, Warning, TEXT("Input at index '%d' (from template T%d) is of type '%s' witch is not supported for that input."), Idx, TemplateIdx, *TargetType);
+			bAreInputValid = false;
+		}
+	}
+	return bAreInputValid;
+}
+void FInputValidator::SetTemplateCount(int TemplateCount)
+{
+	TemplateTypes.SetNum(TemplateCount);
+}
+void FInputValidator::AddSupportedType(ENNETensorDataType Type, int TemplateIdx)
+{
+	check(TemplateTypes.Num() > TemplateIdx);
+	TemplateTypes[TemplateIdx].Add(Type);
+}
+void FInputValidator::AddOptional(int32 TemplateIdx)
+{
+	InputTemplateIndices.Add(TemplateIdx);
+	++NumOptionalInput;
+}
+
+void FInputValidator::AddRequired(int32 TemplateIdx)
+{
+	checkf(NumOptionalInput==0, TEXT("All required attribute should be declared before the optional ones as they are referenced by indices"));
+	InputTemplateIndices.Add(TemplateIdx);
+	++NumRequiredInput;
+}
+
+
+//
+//
+//
+void FAttributeValidator::AddOptional(const FString& Name, ENNEAttributeDataType Type)
+{
+	checkf(nullptr == OptionalAttributes.FindByPredicate([Name](const FEntry& Other) { return Other.Name == Name; }), TEXT("Attribute name should be unique"));
+	checkf(nullptr == RequiredAttributes.FindByPredicate([Name](const FEntry& Other) { return Other.Name == Name; }), TEXT("Attribute name should be unique"));
+	OptionalAttributes.Emplace(Name, Type);
+}
+
+void FAttributeValidator::AddRequired(const FString& Name, ENNEAttributeDataType Type)
+{
+	checkf(nullptr == OptionalAttributes.FindByPredicate([Name](const FEntry& Other) { return Other.Name == Name; }), TEXT("Attribute name should be unique"));
+	checkf(nullptr == RequiredAttributes.FindByPredicate([Name](const FEntry& Other) { return Other.Name == Name; }), TEXT("Attribute name should be unique"));
+	RequiredAttributes.Emplace(Name, Type);
+}
+
+bool FAttributeValidator::Validate(const NNECore::FAttributeMap& AttributesToValidate)
+{
+	bool bAreAttributesValid = true;
+
+	//Verify all required attribute are matching specifications
+	for (int32 Idx = 0; Idx < RequiredAttributes.Num(); ++Idx)
+	{
+		const FNNEAttributeValue* FoundAttribute = AttributesToValidate.GetAttributeValue(RequiredAttributes[Idx].Name);
+		
+		if (FoundAttribute == nullptr)
+		{
+			bAreAttributesValid = false;
+			UE_LOG(LogNNE, Warning, TEXT("Required attribute '%s' not found."),
+				*RequiredAttributes[Idx].Name);
+		}
+		else if (RequiredAttributes[Idx].Type != FoundAttribute->GetType())
+		{
+			bAreAttributesValid = false;
+			UE_LOG(LogNNE, Warning, TEXT("Required attribute '%s' type '%d' does not match expected type '%d'."),
+				*RequiredAttributes[Idx].Name,
+				(int)FoundAttribute->GetType(),
+				(int)RequiredAttributes[Idx].Type);
+		}
+	}
+
+	//Verify all optional attribute are matching specifications
+	for (int32 Idx = 0; Idx < OptionalAttributes.Num(); ++Idx)
+	{
+		const FNNEAttributeValue* FoundAttribute = AttributesToValidate.GetAttributeValue(OptionalAttributes[Idx].Name);
+		
+		if ((FoundAttribute != nullptr) && (OptionalAttributes[Idx].Type != FoundAttribute->GetType()))
+		{
+			bAreAttributesValid = false;
+			UE_LOG(LogNNE, Warning, TEXT("Optional attribute '%s' type '%d' does not match expected type '%d'."),
+				*OptionalAttributes[Idx].Name,
+				(int)FoundAttribute->GetType(),
+				(int)OptionalAttributes[Idx].Type);
+		}
+	}
+
+	//Verify all attributes are either required or optional, otherwise they are unsupported
+	for (int32 Idx = 0; Idx < AttributesToValidate.Num(); ++Idx)
+	{
+		const FString& Name = AttributesToValidate.GetName(Idx);
+		const FEntry* OptionalAttribute = OptionalAttributes.FindByPredicate([Name](const FEntry& Other) { return Other.Name == Name; });
+		const FEntry* RequiredAttribute = RequiredAttributes.FindByPredicate([Name](const FEntry& Other) { return Other.Name == Name; });
+		
+		if (OptionalAttribute == nullptr && RequiredAttribute == nullptr)
+		{
+			bAreAttributesValid = false;
+			UE_LOG(LogNNE, Warning, TEXT("Found unsupported attribute '%s'."), *Name);
+		}
+	}
+
+	return bAreAttributesValid;
+}
+
+} // UE::NNERuntimeRDG::Private
diff --git a/Engine/Plugins/Experimental/NNERuntimeRDG/Source/NNERuntimeRDG/Private/NNERuntimeRDGBase.h b/Engine/Plugins/Experimental/NNERuntimeRDG/Source/NNERuntimeRDG/Private/NNERuntimeRDGBase.h
new file mode 100644
--- /dev/null
+++ b/Engine/Plugins/Experimental/NNERuntimeRDG/Source/NNERuntimeRDG/Private/NNERuntimeRDGBase.h
@@ -0,0 +1,252 @@
+// Copyright Epic Games, Inc. All Rights Reserved.
+
+#pragma once
+
+#include "NNEAttributeMap.h"
+#include "NNEAttributeValue.h"
+#include "NNETensor.h"
+#include "NNETypes.h"
+#include "NNE.h"
+#include "NNEModelOptimizerInterface.h"
+#include "NNERuntimeFormat.h"
+#include "RenderGraphResources.h"
+#include "Serialization/MemoryReader.h"
+#include "ShaderParameterUtils.h"
+
+
+
+BEGIN_SHADER_PARAMETER_STRUCT(FNNETensorReadbackParameters, )
+	RDG_BUFFER_ACCESS(Buffer, ERHIAccess::CopySrc)
+END_SHADER_PARAMETER_STRUCT()
+
+class FRDGBuilder;
+struct FNNEModelRaw;
+
+namespace UE::NNECore { class FTensorDesc; }
+
+namespace UE::NNERuntimeRDG::Private
+{
+
+/**
+ * Interface for all operators to prepare the model tensors at scheduling time
+ */
+struct IPrepareOperator
+{
+	virtual ~IPrepareOperator() = default;
+	virtual int PrepareOutputs(TConstArrayView<NNECore::Internal::FTensorRef> InputTensors, TArrayView<NNECore::Internal::FTensorRef> OutputTensors) const = 0;
+};
+
+/**
+* Interface for all ML operators running on the RDG 
+*/
+struct IOperatorRDG
+{
+	virtual ~IOperatorRDG() = default;
+};
+
+class FTensorRDG : public NNECore::Internal::FTensor
+{
+	FRDGBufferRef Buffer{};
+	
+public:
+	static FTensorRDG Make(const NNECore::FTensorDesc& TensorDesc, const NNECore::FTensorShape& Shape, FRDGBufferRef Buffer)
+	{
+		check(Shape.IsCompatibleWith(TensorDesc.GetShape()));
+		FTensorRDG TensorRDG;
+		TensorRDG.Buffer = nullptr;
+		TensorRDG.Name = TensorDesc.GetName();
+		TensorRDG.DataType = TensorDesc.GetDataType();
+		TensorRDG.Shape = Shape;
+		TensorRDG.Volume = Shape.Volume();
+		TensorRDG.DataSize = (uint64)NNECore::GetTensorDataTypeSizeInBytes(TensorRDG.DataType) * TensorRDG.Volume;
+		TensorRDG.Buffer = Buffer;
+		check(TensorRDG.Volume <= TNumericLimits<uint32>::Max());
+		return TensorRDG;
+	}
+
+	bool HasBuffer() const { return Buffer != FRDGBufferRef{}; }
+	void SetBuffer(FRDGBufferRef Inbuffer){ Buffer = Inbuffer; }
+	FRDGBufferRef GetBuffer() const { return Buffer; }
+};
+
+using FTensorRDGRef = FTensorRDG*;
+using FTensorRDGArray = TArray<FTensorRDG, TInlineAllocator<16>>;
+using FTensorRDGRefArray = TArray<FTensorRDGRef, TInlineAllocator<64>>;
+using FIntArray = TArray<int32, TInlineAllocator<16>>;
+
+bool AlwaysValidValidationFunction(const NNECore::FAttributeMap& AttributeMap, TConstArrayView<ENNETensorDataType> InputTypes, TConstArrayView<NNECore::FSymbolicTensorShape> InputShapes);
+
+class FInputValidator
+{
+public:
+	FInputValidator();
+	void AddRequired(int32 TemplateIdx = 0);
+	void AddOptional(int32 TemplateIdx = 0);
+	void SetTemplateCount(int32 TemplateCount);
+	void AddSupportedType(ENNETensorDataType Type, int32 TemplateIdx = 0);
+	bool Validate(TConstArrayView<ENNETensorDataType> InputTypes);
+
+private:
+	TArray<TArray<ENNETensorDataType>> TemplateTypes;
+	TArray<int32> InputTemplateIndices;
+	int32 NumRequiredInput;
+	int32 NumOptionalInput;
+};
+
+class FAttributeValidator
+{
+public:
+	void AddOptional(const FString& Name, ENNEAttributeDataType Type);
+	void AddRequired(const FString& Name, ENNEAttributeDataType Type);
+	bool Validate(const NNECore::FAttributeMap& AttributesToValidate);
+
+private:
+	struct FEntry
+	{
+		FEntry(const FString& InName, ENNEAttributeDataType InType)
+			: Name(InName), Type(InType)
+		{
+		}
+
+		//Idea: we could extended as needed by operator to support more validation especially around the range of the values.
+		//An example is ConvTranspose `auto_pad` enum style string that can only take a few values
+		//In the same direction we might only support a range of value for a float (for example
+		//we only support integer but the type is float, or only positive values for an int32)
+		FString Name;
+		ENNEAttributeDataType Type;
+	};
+
+	TArray<FEntry> RequiredAttributes;
+	TArray<FEntry> OptionalAttributes;
+};
+
+/**
+ * Registry for RDG ML operators
+ */
+template<class TOperatorType>
+class TOperatorRegistryRDG
+{
+public:
+
+	typedef TOperatorType* (*OperatorCreateFunc)();
+	typedef bool (*OperatorValidateFunc)(const NNECore::FAttributeMap& AttributeMap, TConstArrayView<ENNETensorDataType> InputTypes, TConstArrayView<NNECore::FSymbolicTensorShape> InputShapes);
+
+	static TOperatorRegistryRDG* Get()
+	{
+		static TOperatorRegistryRDG Inst;
+
+		return &Inst;
+	}
+
+	OperatorValidateFunc OpFindValidation(const FString& Name)
+	{
+		OperatorValidateFunc* Fn = OperatorValidations.Find(Name);
+
+		if (!Fn)
+		{
+			UE_LOG(LogNNE, Warning, TEXT("RDG MLOperator:%s is not registered"), *Name);
+			return nullptr;
+		}
+
+		return *Fn;
+	}
+
+	OperatorCreateFunc OpFind(const FString& Name)
+	{
+		OperatorCreateFunc* Fn = Operators.Find(Name);
+
+		if (!Fn)
+		{
+			UE_LOG(LogNNE, Warning, TEXT("RDG MLOperator:%s is not registered"), *Name);
+			return nullptr;
+		}
+
+		return *Fn;
+	}
+
+	bool OpAdd(const FString& Name, OperatorCreateFunc Func, OperatorValidateFunc ValidateFunc = AlwaysValidValidationFunction)
+	{
+		if (Operators.Find(Name) != nullptr)
+		{
+			UE_LOG(LogNNE, Warning, TEXT("RDG MLOperator is already registered:%s"), *Name);
+			return false;
+		}
+
+		Operators.Add(Name, Func);
+		OperatorValidations.Add(Name, ValidateFunc);
+		return true;
+	}
+
+private:
+
+	TMap<FString, OperatorCreateFunc> Operators;
+	TMap<FString, OperatorValidateFunc> OperatorValidations;
+};
+
+/**
+ * Validator for RDG ML operators
+ */
+template<class TOperatorType>
+class TModelValidatorRDG : public UE::NNECore::Internal::IModelValidator
+{
+public:
+	virtual FString GetName() const 
+	{
+		return TEXT("RDG Model validator");
+	}
+	
+	virtual bool ValidateModel(const FNNEModelRaw& InputModel, const UE::NNECore::Internal::FOptimizerOptionsMap& Options) const override
+	{
+		FNNERuntimeFormat	Format;
+
+		ENNEInferenceFormat FormatType = InputModel.Format;
+		if (FormatType != ENNEInferenceFormat::NNERT)
+		{
+			UE_LOG(LogNNE, Warning, TEXT("Unsupported format type for validator %s"), *GetName());
+			return false;
+		}
+
+		FMemoryReader Reader(InputModel.Data);
+		FNNERuntimeFormat::StaticStruct()->SerializeBin(Reader, &Format);
+
+		TOperatorRegistryRDG<TOperatorType>* Registry = TOperatorRegistryRDG<TOperatorType>::Get();
+		check(Registry != nullptr);
+
+		for (int32 Idx = 0; Idx < Format.Operators.Num(); ++Idx)
+		{
+			TArray<ENNETensorDataType> InputTensorTypes;
+			TArray<NNECore::FSymbolicTensorShape> InputTensorShapes;
+			NNECore::FAttributeMap AttributeMap;
+			
+			for (int32 InputTensorIndex: Format.Operators[Idx].InTensors)
+			{
+				InputTensorTypes.Add(Format.Tensors[InputTensorIndex].DataType);
+				InputTensorShapes.Add(NNECore::FSymbolicTensorShape::Make(Format.Tensors[InputTensorIndex].Shape));
+			}
+			for (const FNNEFormatAttributeDesc& Desc : Format.Operators[Idx].Attributes)
+			{
+				AttributeMap.SetAttribute(Desc.Name, Desc.Value);
+			}
+
+			const FString& OpType = Format.Operators[Idx].TypeName;
+			
+			typename TOperatorRegistryRDG<TOperatorType>::OperatorValidateFunc ValidationFn = Registry->OpFindValidation(OpType);
+
+			if (!ValidationFn)
+			{
+				UE_LOG(LogNNE, Warning, TEXT("Hlsl MLOperatorRegistry failed to find validation for operator:%s"), *OpType);
+				return false;
+			}
+			
+			if (!ValidationFn(AttributeMap, InputTensorTypes, InputTensorShapes))
+			{
+				UE_LOG(LogNNE, Warning, TEXT("Hlsl MLOperatorRegistry failed to validate operator:%s"), *OpType);
+				return false;
+			}
+		}
+
+		return true;
+	}
+};
+
+} // UE::NNERuntimeRDG::Private
\ No newline at end of file
diff --git a/Engine/Plugins/Experimental/NNERuntimeRDG/Source/NNERuntimeRDG/Private/NNERuntimeRDGDml.cpp b/Engine/Plugins/Experimental/NNERuntimeRDG/Source/NNERuntimeRDG/Private/NNERuntimeRDGDml.cpp
--- a/Engine/Plugins/Experimental/NNERuntimeRDG/Source/NNERuntimeRDG/Private/NNERuntimeRDGDml.cpp
+++ b/Engine/Plugins/Experimental/NNERuntimeRDG/Source/NNERuntimeRDG/Private/NNERuntimeRDGDml.cpp
@@ -4,8 +4,8 @@
 
 #include "NNERuntimeRDG.h"
 #include "NNEUtilsModelOptimizer.h"
-#include "NNECoreModelData.h"
-#include "NNECoreRuntimeFormat.h"
+#include "NNEModelData.h"
+#include "NNERuntimeFormat.h"
 #include "HAL/FileManager.h"
 #include "Interfaces/IPluginManager.h"
 #include "Algo/Find.h"
diff --git a/Engine/Plugins/Experimental/NNERuntimeRDG/Source/NNERuntimeRDG/Private/NNERuntimeRDGDml.h b/Engine/Plugins/Experimental/NNERuntimeRDG/Source/NNERuntimeRDG/Private/NNERuntimeRDGDml.h
--- a/Engine/Plugins/Experimental/NNERuntimeRDG/Source/NNERuntimeRDG/Private/NNERuntimeRDGDml.h
+++ b/Engine/Plugins/Experimental/NNERuntimeRDG/Source/NNERuntimeRDG/Private/NNERuntimeRDGDml.h
@@ -3,8 +3,8 @@
 #pragma once
 
 #include "CoreMinimal.h"
-#include "NNECoreRuntime.h"
-#include "NNECoreRuntimeRDG.h"
+#include "NNERuntime.h"
+#include "NNERuntimeRDG.h"
 #include "Templates/UniquePtr.h"
 #include "UObject/Class.h"
 #include "UObject/Object.h"
diff --git a/Engine/Plugins/Experimental/NNERuntimeRDG/Source/NNERuntimeRDG/Private/NNERuntimeRDGHlsl.h b/Engine/Plugins/Experimental/NNERuntimeRDG/Source/NNERuntimeRDG/Private/NNERuntimeRDGHlsl.h
--- a/Engine/Plugins/Experimental/NNERuntimeRDG/Source/NNERuntimeRDG/Private/NNERuntimeRDGHlsl.h
+++ b/Engine/Plugins/Experimental/NNERuntimeRDG/Source/NNERuntimeRDG/Private/NNERuntimeRDGHlsl.h
@@ -3,8 +3,8 @@
 #pragma once
 
 #include "CoreMinimal.h"
-#include "NNECoreRuntime.h"
-#include "NNECoreRuntimeRDG.h"
+#include "NNERuntime.h"
+#include "NNERuntimeRDG.h"
 #include "Templates/UniquePtr.h"
 #include "UObject/Class.h"
 #include "UObject/Object.h"
diff --git a/Engine/Plugins/Experimental/NNERuntimeRDG/Source/NNERuntimeRDG/Private/NNERuntimeRDGModel.cpp b/Engine/Plugins/Experimental/NNERuntimeRDG/Source/NNERuntimeRDG/Private/NNERuntimeRDGModel.cpp
--- a/Engine/Plugins/Experimental/NNERuntimeRDG/Source/NNERuntimeRDG/Private/NNERuntimeRDGModel.cpp
+++ b/Engine/Plugins/Experimental/NNERuntimeRDG/Source/NNERuntimeRDG/Private/NNERuntimeRDGModel.cpp
@@ -2,7 +2,7 @@
 
 #include "NNERuntimeRDGModel.h"
 
-#include "NNECoreRuntimeFormat.h"
+#include "NNERuntimeFormat.h"
 #include "NNEUtilsModelOptimizer.h"
 #include "RenderGraphBuilder.h"
 #include "RenderGraphUtils.h"
diff --git a/Engine/Plugins/Experimental/NNERuntimeRDG/Source/NNERuntimeRDG/Private/NNERuntimeRDGModel.h b/Engine/Plugins/Experimental/NNERuntimeRDG/Source/NNERuntimeRDG/Private/NNERuntimeRDGModel.h
--- a/Engine/Plugins/Experimental/NNERuntimeRDG/Source/NNERuntimeRDG/Private/NNERuntimeRDGModel.h
+++ b/Engine/Plugins/Experimental/NNERuntimeRDG/Source/NNERuntimeRDG/Private/NNERuntimeRDGModel.h
@@ -3,11 +3,11 @@
 #pragma once
 
 #include "CoreMinimal.h"
-#include "NNECoreModelBase.h"
-#include "NNECoreRuntimeRDG.h"
-#include "NNECoreTensor.h"
-#include "NNECoreTypes.h"
+#include "NNEModelBase.h"
 #include "NNERuntimeRDG.h"
+#include "NNETensor.h"
+#include "NNETypes.h"
+#include "NNERuntimeRDGBase.h"
 
 namespace UE::NNERuntimeRDG::Private
 {
