This CL ads insights on NNE usage inside the editor. Events are generated for ...

- ... StartupModule
  - PlatformName -> ['Windows', 'Linux', 'MacOs']
  - ModuleName -> ['NNE']
- ... FactoryCreateBinary
  - PlatformName -> ['Windows', 'Linux', 'MacOs']
  - FactoryName -> ['UNNEModelDataFactory']
  - ModelFileSize -> [<int>]
- ... RegisterRuntime
  - PlatformName ->['Windows', 'Linux', 'MacOs']
  - HashedRuntimeName -> [<MD5 hash of the runtime name>]
- ... CreateModel
  - PlatformName ->['Windows', 'Linux', 'MacOs']
  - RuntimeName -> [<MD5 hash of the runtime name>]
  - ModelDataSize -> [<int>]

This will allow us to create:
- Piechart showing usage per platform (for NNE as whole, but also for each of the following)
- Statistics on model file size (mean, variance, ...)
- Pichart showing usage per runtime
- Statistics on optimized model sizes (mean, variance, ...)

Note: All runtime names are now one-way hashed so we do not collect data on names created by licensees. In the analytics, we still can replace hashes of known runtime names by their readable version.

Tested by receiving events in unrealtime-services.

#rnx

diff --git a/Engine/Plugins/Experimental/NNE/Source/NNE/Private/NNE.cpp b/Engine/Plugins/Experimental/NNE/Source/NNE/Private/NNE.cpp
--- a/Engine/Plugins/Experimental/NNE/Source/NNE/Private/NNE.cpp
+++ b/Engine/Plugins/Experimental/NNE/Source/NNE/Private/NNE.cpp
@@ -2,6 +2,11 @@
 
 #include "NNE.h"
 
+#include "EngineAnalytics.h"
+#include "Kismet/GameplayStatics.h"
+#include "Misc/CoreDelegates.h"
+#include "Misc/SecureHash.h"
+
 DEFINE_LOG_CATEGORY(LogNNE);
 
 namespace UE::NNECore
@@ -89,6 +94,18 @@ namespace UE::NNECore
 
 	bool RegisterRuntime(TWeakInterfacePtr<INNERuntime> Runtime)
 	{
+		FString RuntimeName = Runtime->GetRuntimeName();
+		FCoreDelegates::OnAllModuleLoadingPhasesComplete.AddLambda([RuntimeName]() {
+			if (FEngineAnalytics::IsAvailable())
+			{
+				TArray<FAnalyticsEventAttribute> Attributes = MakeAnalyticsEventAttributeArray(
+					TEXT("PlatformName"), UGameplayStatics::GetPlatformName(),
+					TEXT("HashedRuntimeName"), FMD5::HashAnsiString(*RuntimeName)
+				);
+				FEngineAnalytics::GetProvider().RecordEvent(TEXT("NeuralNetworkEngine.RegisterRuntime"), Attributes);
+			}
+		});
+
 		return FRegistry::Get()->Add(Runtime);
 	}
 
diff --git a/Engine/Plugins/Experimental/NNE/Source/NNE/Private/NNEModule.cpp b/Engine/Plugins/Experimental/NNE/Source/NNE/Private/NNEModule.cpp
--- a/Engine/Plugins/Experimental/NNE/Source/NNE/Private/NNEModule.cpp
+++ b/Engine/Plugins/Experimental/NNE/Source/NNE/Private/NNEModule.cpp
@@ -1,5 +1,8 @@
 // Copyright Epic Games, Inc. All Rights Reserved.
 
+#include "EngineAnalytics.h"
+#include "Kismet/GameplayStatics.h"
+#include "Misc/CoreDelegates.h"
 #include "Modules/ModuleManager.h"
 
 class FNNECoreModule : public IModuleInterface
@@ -8,6 +11,16 @@ public:
 
 	virtual void StartupModule() override
 	{
+		FCoreDelegates::OnAllModuleLoadingPhasesComplete.AddLambda([]() {
+			if (FEngineAnalytics::IsAvailable())
+			{
+				TArray<FAnalyticsEventAttribute> Attributes = MakeAnalyticsEventAttributeArray(
+					TEXT("PlatformName"), UGameplayStatics::GetPlatformName(),
+					TEXT("ModuleName"), TEXT("NNE")
+				);
+				FEngineAnalytics::GetProvider().RecordEvent(TEXT("NeuralNetworkEngine.StartupModule"), Attributes);
+			}
+		});
 	}
 
 	virtual void ShutdownModule() override
diff --git a/Engine/Plugins/Experimental/NNE/Source/NNEEditorTools/Private/NNEModelDataFactory.cpp b/Engine/Plugins/Experimental/NNE/Source/NNEEditorTools/Private/NNEModelDataFactory.cpp
--- a/Engine/Plugins/Experimental/NNE/Source/NNEEditorTools/Private/NNEModelDataFactory.cpp
+++ b/Engine/Plugins/Experimental/NNE/Source/NNEEditorTools/Private/NNEModelDataFactory.cpp
@@ -1,10 +1,13 @@
 // Copyright Epic Games, Inc. All Rights Reserved.
 
 #include "NNEModelDataFactory.h"
+
 #include "Editor.h"
 #include "Editor/EditorEngine.h"
 #include "EditorFramework/AssetImportData.h"
+#include "EngineAnalytics.h"
 #include "HAL/FileManager.h"
+#include "Kismet/GameplayStatics.h"
 #include "NNE.h"
 #include "NNEModelData.h"
 #include "Subsystems/ImportSubsystem.h"
@@ -36,6 +39,16 @@ UObject* UNNEModelDataFactory::FactoryCreateBinary(UClass* Class, UObject* InPar
 	ModelData->AssetImportData->Update(GetCurrentFilename());
 	GEditor->GetEditorSubsystem<UImportSubsystem>()->BroadcastAssetPostImport(this, ModelData);
 
+	if (FEngineAnalytics::IsAvailable())
+	{
+		TArray<FAnalyticsEventAttribute> Attributes = MakeAnalyticsEventAttributeArray(
+			TEXT("PlatformName"), UGameplayStatics::GetPlatformName(),
+			TEXT("FactoryName"), TEXT("UNNEModelDataFactory"),
+			TEXT("ModelFileSize"), BufferView.Num()
+		);
+		FEngineAnalytics::GetProvider().RecordEvent(TEXT("NeuralNetworkEngine.FactoryCreateBinary"), Attributes);
+	}
+
 	return ModelData;
 }
 
diff --git a/Engine/Plugins/Experimental/NNERuntimeORTCpu/Source/NNERuntimeORTCpu/Private/NNERuntimeORTCpu.cpp b/Engine/Plugins/Experimental/NNERuntimeORTCpu/Source/NNERuntimeORTCpu/Private/NNERuntimeORTCpu.cpp
--- a/Engine/Plugins/Experimental/NNERuntimeORTCpu/Source/NNERuntimeORTCpu/Private/NNERuntimeORTCpu.cpp
+++ b/Engine/Plugins/Experimental/NNERuntimeORTCpu/Source/NNERuntimeORTCpu/Private/NNERuntimeORTCpu.cpp
@@ -2,6 +2,9 @@
 
 #include "NNERuntimeORTCpu.h"
 
+#include "EngineAnalytics.h"
+#include "Kismet/GameplayStatics.h"
+#include "Misc/SecureHash.h"
 #include "NNEAttributeMap.h"
 #include "NNEModelData.h"
 #include "NNEModelOptimizerInterface.h"
@@ -77,5 +80,16 @@ TUniquePtr<UE::NNECore::IModelCPU> UNNERuntimeORTCpuImpl::CreateModel(TObjectPtr
 	TConstArrayView<uint8> Data = ModelData->GetModelData(GetRuntimeName());
 	UE::NNERuntimeORTCpu::Private::FModelCPU* Model = new UE::NNERuntimeORTCpu::Private::FModelCPU(&NNEEnvironmentCPU, Data);
 	UE::NNECore::IModelCPU* IModel = static_cast<UE::NNECore::IModelCPU*>(Model);
+
+	if (FEngineAnalytics::IsAvailable())
+	{
+		TArray<FAnalyticsEventAttribute> Attributes = MakeAnalyticsEventAttributeArray(
+			TEXT("PlatformName"), UGameplayStatics::GetPlatformName(),
+			TEXT("HashedRuntimeName"), FMD5::HashAnsiString(*GetRuntimeName()),
+			TEXT("ModelDataSize"), Data.Num()
+		);
+		FEngineAnalytics::GetProvider().RecordEvent(TEXT("NeuralNetworkEngine.CreateModel"), Attributes);
+	}
+
 	return TUniquePtr<UE::NNECore::IModelCPU>(IModel);
 }
\ No newline at end of file
diff --git a/Engine/Plugins/Experimental/NNERuntimeORTGpu/Source/NNERuntimeORT/Private/NNERuntimeORT.cpp b/Engine/Plugins/Experimental/NNERuntimeORTGpu/Source/NNERuntimeORT/Private/NNERuntimeORT.cpp
--- a/Engine/Plugins/Experimental/NNERuntimeORTGpu/Source/NNERuntimeORT/Private/NNERuntimeORT.cpp
+++ b/Engine/Plugins/Experimental/NNERuntimeORTGpu/Source/NNERuntimeORT/Private/NNERuntimeORT.cpp
@@ -1,6 +1,10 @@
 // Copyright Epic Games, Inc. All Rights Reserved.
 
 #include "NNERuntimeORT.h"
+
+#include "EngineAnalytics.h"
+#include "Kismet/GameplayStatics.h"
+#include "Misc/SecureHash.h"
 #include "NNERuntimeORTModel.h"
 #include "NNERuntimeORTUtils.h"
 #include "NNEUtilsModelOptimizer.h"
@@ -115,6 +119,16 @@ TUniquePtr<UE::NNECore::IModelGPU> UNNERuntimeORTGpuImpl::CreateModel(TObjectPtr
 			return TUniquePtr<UE::NNECore::IModelGPU>();
 	}
 
+	if (FEngineAnalytics::IsAvailable())
+	{
+		TArray<FAnalyticsEventAttribute> Attributes = MakeAnalyticsEventAttributeArray(
+			TEXT("PlatformName"), UGameplayStatics::GetPlatformName(),
+			TEXT("HashedRuntimeName"), FMD5::HashAnsiString(*GetRuntimeName()),
+			TEXT("ModelDataSize"), Data.Num()
+		);
+		FEngineAnalytics::GetProvider().RecordEvent(TEXT("NeuralNetworkEngine.CreateModel"), Attributes);
+	}
+
 	return TUniquePtr<UE::NNECore::IModelGPU>(IModel);
 }
 
diff --git a/Engine/Plugins/Experimental/NNERuntimeRDG/Source/NNERuntimeRDG/Private/NNERuntimeRDGDml.cpp b/Engine/Plugins/Experimental/NNERuntimeRDG/Source/NNERuntimeRDG/Private/NNERuntimeRDGDml.cpp
--- a/Engine/Plugins/Experimental/NNERuntimeRDG/Source/NNERuntimeRDG/Private/NNERuntimeRDGDml.cpp
+++ b/Engine/Plugins/Experimental/NNERuntimeRDG/Source/NNERuntimeRDG/Private/NNERuntimeRDGDml.cpp
@@ -2,15 +2,17 @@
 
 #include "NNERuntimeRDGDml.h"
 
-#include "NNERuntimeRDG.h"
-#include "NNEUtilsModelOptimizer.h"
-#include "NNEModelData.h"
-#include "NNERuntimeFormat.h"
-#include "HAL/FileManager.h"
-#include "Interfaces/IPluginManager.h"
 #include "Algo/Find.h"
-
 #include "Dml/NNEDmlModel.h"
+#include "EngineAnalytics.h"
+#include "HAL/FileManager.h"
+#include "Interfaces/IPluginManager.h"
+#include "Kismet/GameplayStatics.h"
+#include "Misc/SecureHash.h"
+#include "NNEModelData.h"
+#include "NNERuntimeFormat.h"
+#include "NNERuntimeRDG.h"
+#include "NNEUtilsModelOptimizer.h"
 
 #ifdef NNE_USE_DIRECTML
 
@@ -158,6 +160,16 @@ TUniquePtr<UE::NNECore::IModelRDG> UNNERuntimeRDGDmlImpl::CreateModel(TObjectPtr
 	TConstArrayView<uint8> Data = ModelData->GetModelData(GetRuntimeName());
 	FModel* Model = new FModel(Data, Ctx);
 
+	if (FEngineAnalytics::IsAvailable())
+	{
+		TArray<FAnalyticsEventAttribute> Attributes = MakeAnalyticsEventAttributeArray(
+			TEXT("PlatformName"), UGameplayStatics::GetPlatformName(),
+			TEXT("HashedRuntimeName"), FMD5::HashAnsiString(*GetRuntimeName()),
+			TEXT("ModelDataSize"), Data.Num()
+		);
+		FEngineAnalytics::GetProvider().RecordEvent(TEXT("NeuralNetworkEngine.CreateModel"), Attributes);
+	}
+
 	return TUniquePtr<UE::NNECore::IModelRDG>(Model);
 #else
 	return TUniquePtr<UE::NNECore::IModelRDG>();
diff --git a/Engine/Plugins/Experimental/NNERuntimeRDG/Source/NNERuntimeRDG/Private/NNERuntimeRDGHlsl.cpp b/Engine/Plugins/Experimental/NNERuntimeRDG/Source/NNERuntimeRDG/Private/NNERuntimeRDGHlsl.cpp
--- a/Engine/Plugins/Experimental/NNERuntimeRDG/Source/NNERuntimeRDG/Private/NNERuntimeRDGHlsl.cpp
+++ b/Engine/Plugins/Experimental/NNERuntimeRDG/Source/NNERuntimeRDG/Private/NNERuntimeRDGHlsl.cpp
@@ -1,6 +1,10 @@
 // Copyright Epic Games, Inc. All Rights Reserved.
 
 #include "NNERuntimeRDGHlsl.h"
+
+#include "EngineAnalytics.h"
+#include "Kismet/GameplayStatics.h"
+#include "Misc/SecureHash.h"
 #include "NNEAttributeMap.h"
 #include "NNEModelData.h"
 #include "NNEModelOptimizerInterface.h"
@@ -128,5 +132,15 @@ TUniquePtr<UE::NNECore::IModelRDG> UNNERuntimeRDGHlslImpl::CreateModel(TObjectPt
 	TConstArrayView<uint8> Data = ModelData->GetModelData(GetRuntimeName());
 	UE::NNERuntimeRDG::Private::Hlsl::FModel* Model = new UE::NNERuntimeRDG::Private::Hlsl::FModel(Data);
 
+	if (FEngineAnalytics::IsAvailable())
+	{
+		TArray<FAnalyticsEventAttribute> Attributes = MakeAnalyticsEventAttributeArray(
+			TEXT("PlatformName"), UGameplayStatics::GetPlatformName(),
+			TEXT("HashedRuntimeName"), FMD5::HashAnsiString(*GetRuntimeName()),
+			TEXT("ModelDataSize"), Data.Num()
+		);
+		FEngineAnalytics::GetProvider().RecordEvent(TEXT("NeuralNetworkEngine.CreateModel"), Attributes);
+	}
+
 	return TUniquePtr<UE::NNECore::IModelRDG>(Model);
 }
