NNE Rename NNECore module to NNE following split of runtimes into there own plugins (step2)

diff --git a/Engine/Plugins/Experimental/NNERuntimeRDG/Source/NNERuntimeRDG/Private/Dml/NNEDmlModel.cpp b/Engine/Plugins/Experimental/NNERuntimeRDG/Source/NNERuntimeRDG/Private/Dml/NNEDmlModel.cpp
--- a/Engine/Plugins/Experimental/NNERuntimeRDG/Source/NNERuntimeRDG/Private/Dml/NNEDmlModel.cpp
+++ b/Engine/Plugins/Experimental/NNERuntimeRDG/Source/NNERuntimeRDG/Private/Dml/NNEDmlModel.cpp
@@ -321,7 +321,7 @@ public:
 
 	struct FGraphDesc
 	{
-		TConstArrayView<NNECore::Internal::FTensor>	AllTensors;
+		TConstArrayView<NNE::Internal::FTensor>	AllTensors;
 		TConstArrayView<int32>		InputIndices;
 		TConstArrayView<int32>		OutputIndices;
 		TConstArrayView<int32>		WeightIndices;
@@ -862,12 +862,12 @@ bool FModelInstance::Init(TConstArrayView<uint8> ModelData, FDmlDeviceContext* I
 
 	// DirectML requires all tensors to be concrete
 	// Notes: to handle dynamic tensor desc, op should init from symbolic shapes
-	TArray<NNECore::Internal::FTensor>	Tensors;
+	TArray<NNE::Internal::FTensor>	Tensors;
 
 	Tensors.Reset(AllSymbolicTensorDescs.Num());
-	for (const NNECore::FTensorDesc& TensorDesc : AllSymbolicTensorDescs)
+	for (const NNE::FTensorDesc& TensorDesc : AllSymbolicTensorDescs)
 	{
-		Tensors.Emplace(NNECore::Internal::FTensor::MakeFromSymbolicDesc(TensorDesc));
+		Tensors.Emplace(NNE::Internal::FTensor::MakeFromSymbolicDesc(TensorDesc));
 	}
 
 	TArray<FGraphBuilder::FOpDesc>	DmlGraphOperators;
@@ -881,9 +881,9 @@ bool FModelInstance::Init(TConstArrayView<uint8> ModelData, FDmlDeviceContext* I
 		const FString TypeName = Format.Operators[Idx].TypeName;
 
 		FGraphBuilder::FOpDesc					OpDesc;
-		TArray<NNECore::Internal::FTensor>		OpInputTensors;
-		TArray<NNECore::Internal::FTensor>		OpOutputTensors;
-		NNECore::FAttributeMap					AttributeMap;
+		TArray<NNE::Internal::FTensor>		OpInputTensors;
+		TArray<NNE::Internal::FTensor>		OpOutputTensors;
+		NNE::FAttributeMap					AttributeMap;
 
 		OpDesc.InputStart = OpInputIndices.Num();
 		OpDesc.OutputStart = OpOutputIndices.Num();
@@ -900,9 +900,9 @@ bool FModelInstance::Init(TConstArrayView<uint8> ModelData, FDmlDeviceContext* I
 			}
 			else
 			{
-				NNECore::FTensorDesc SymbolicTensorDesc = AllSymbolicTensorDescs[InputTensorIndex];
+				NNE::FTensorDesc SymbolicTensorDesc = AllSymbolicTensorDescs[InputTensorIndex];
 				// Notes: to handle dynamic tensor desc, op should init from symbolic shapes
-				OpInputTensors.Emplace(NNECore::Internal::FTensor::MakeFromSymbolicDesc(SymbolicTensorDesc));
+				OpInputTensors.Emplace(NNE::Internal::FTensor::MakeFromSymbolicDesc(SymbolicTensorDesc));
 			}
 
 			OpInputIndices.Emplace(InputTensorIndex);
@@ -910,9 +910,9 @@ bool FModelInstance::Init(TConstArrayView<uint8> ModelData, FDmlDeviceContext* I
 
 		for (int32 OutputTensorIndex : Format.Operators[Idx].OutTensors)
 		{
-			NNECore::FTensorDesc SymbolicTensorDesc = AllSymbolicTensorDescs[OutputTensorIndex];
+			NNE::FTensorDesc SymbolicTensorDesc = AllSymbolicTensorDescs[OutputTensorIndex];
 			// Notes: to handle dynamic tensor desc, op should init from symbolic shapes
-			OpOutputTensors.Emplace(NNECore::Internal::FTensor::MakeFromSymbolicDesc(SymbolicTensorDesc));
+			OpOutputTensors.Emplace(NNE::Internal::FTensor::MakeFromSymbolicDesc(SymbolicTensorDesc));
 			OpOutputIndices.Emplace(OutputTensorIndex);
 		}
 
@@ -1299,7 +1299,7 @@ void FModelInstance::AddDispatchOps_RenderThread(FRDGBuilder& GraphBuilder)
 //
 // Create operator
 //
-FOperatorDml* FModelInstance::OpCreate(const FString& OpName, TArrayView<const NNECore::Internal::FTensor> InputTensorDescs, TArrayView<const NNECore::Internal::FTensor> OutputTensorDescs, const NNECore::FAttributeMap& Attributes)
+FOperatorDml* FModelInstance::OpCreate(const FString& OpName, TArrayView<const NNE::Internal::FTensor> InputTensorDescs, TArrayView<const NNE::Internal::FTensor> OutputTensorDescs, const NNE::FAttributeMap& Attributes)
 {
 	FOperatorRegistryDml::OperatorCreateFunc CreateFn = FOperatorRegistryDml::Get()->OpFind(OpName);
 
@@ -1374,7 +1374,7 @@ ID3D12Resource* FModelInstance::CreateD3D12Buffer(uint32 Size, D3D12_RESOURCE_ST
 //
 int FModelInstance::PrepareTensorShapesAndData()
 {
-	for (NNECore::FTensorDesc SymbolicTensorDesc : AllSymbolicTensorDescs)
+	for (NNE::FTensorDesc SymbolicTensorDesc : AllSymbolicTensorDescs)
 	{
 		if (!SymbolicTensorDesc.GetShape().IsConcrete())
 		{
@@ -1386,18 +1386,18 @@ int FModelInstance::PrepareTensorShapesAndData()
 	return 0;
 }
 
-TUniquePtr<UE::NNECore::IModelInstanceRDG> FModel::CreateModelInstance()
+TUniquePtr<UE::NNE::IModelInstanceRDG> FModel::CreateModelInstance()
 {
 	FModelInstance* ModelInstance = new FModelInstance();
 
 	if (!ModelInstance->Init(ModelData, DevCtx))
 	{
 		delete ModelInstance;
-		return TUniquePtr<UE::NNECore::IModelInstanceRDG>();
+		return TUniquePtr<UE::NNE::IModelInstanceRDG>();
 	}
 
-	UE::NNECore::IModelInstanceRDG* IModelInstance = static_cast<UE::NNECore::IModelInstanceRDG*>(ModelInstance);
-	return TUniquePtr<UE::NNECore::IModelInstanceRDG>(IModelInstance);
+	UE::NNE::IModelInstanceRDG* IModelInstance = static_cast<UE::NNE::IModelInstanceRDG*>(ModelInstance);
+	return TUniquePtr<UE::NNE::IModelInstanceRDG>(IModelInstance);
 }
 
 FModel::FModel(TConstArrayView<uint8> InModelData, FDmlDeviceContext* InDevCtx) 
diff --git a/Engine/Plugins/Experimental/NNERuntimeRDG/Source/NNERuntimeRDG/Private/Dml/NNEDmlModel.h b/Engine/Plugins/Experimental/NNERuntimeRDG/Source/NNERuntimeRDG/Private/Dml/NNEDmlModel.h
--- a/Engine/Plugins/Experimental/NNERuntimeRDG/Source/NNERuntimeRDG/Private/Dml/NNEDmlModel.h
+++ b/Engine/Plugins/Experimental/NNERuntimeRDG/Source/NNERuntimeRDG/Private/Dml/NNEDmlModel.h
@@ -64,7 +64,7 @@ private:
 
 	bool InitCompiledOp(TConstArrayView<int32> OpInputIndices, uint64 TensorDataSize);
 
-	FOperatorDml* OpCreate(const FString& Name, TArrayView<const NNECore::Internal::FTensor> InputTensorDesc, TArrayView<const NNECore::Internal::FTensor> OutputTensorDescs, const NNECore::FAttributeMap& Attributes);
+	FOperatorDml* OpCreate(const FString& Name, TArrayView<const NNE::Internal::FTensor> InputTensorDesc, TArrayView<const NNE::Internal::FTensor> OutputTensorDescs, const NNE::FAttributeMap& Attributes);
 
 	FBufferRHIRef CreateRHIBuffer(FRHICommandListImmediate& RHICmdList, uint32 Size, EBufferUsageFlags Usage, ERHIAccess Access, const TCHAR* DbgName);
 	ID3D12Resource* CreateD3D12Buffer(uint32 Size, D3D12_RESOURCE_STATES ResourceState = D3D12_RESOURCE_STATE_COMMON, D3D12_HEAP_TYPE HeapType = D3D12_HEAP_TYPE_DEFAULT, const TCHAR* DebugName = nullptr);
@@ -97,13 +97,13 @@ private:
 	ID3D12DynamicRHI*					DynamicRHI;
 };
 
-class FModel : public NNECore::IModelRDG
+class FModel : public NNE::IModelRDG
 {
 public:
 	FModel(TConstArrayView<uint8> InModelData, FDmlDeviceContext* InDevCtx);
 	virtual ~FModel() {};
 
-	virtual TUniquePtr<UE::NNECore::IModelInstanceRDG> CreateModelInstance() override;
+	virtual TUniquePtr<UE::NNE::IModelInstanceRDG> CreateModelInstance() override;
 
 private:
 	TArray<uint8> ModelData;
diff --git a/Engine/Plugins/Experimental/NNERuntimeRDG/Source/NNERuntimeRDG/Private/Dml/NNEDmlOperator.cpp b/Engine/Plugins/Experimental/NNERuntimeRDG/Source/NNERuntimeRDG/Private/Dml/NNEDmlOperator.cpp
--- a/Engine/Plugins/Experimental/NNERuntimeRDG/Source/NNERuntimeRDG/Private/Dml/NNEDmlOperator.cpp
+++ b/Engine/Plugins/Experimental/NNERuntimeRDG/Source/NNERuntimeRDG/Private/Dml/NNEDmlOperator.cpp
@@ -18,7 +18,7 @@ static bool IsSameShape(TConstArrayView<uint32> Left, TConstArrayView<uint32> Ri
 	return Algo::Compare(Left, Right);
 }
 
-bool IsSameShape(const NNECore::Internal::FTensor& Left, const NNECore::Internal::FTensor& Right)
+bool IsSameShape(const NNE::Internal::FTensor& Left, const NNE::Internal::FTensor& Right)
 {
 	return IsSameShape(Left.GetShape().GetData(), Right.GetShape().GetData());
 }
@@ -248,7 +248,7 @@ FTensorDescDml& FTensorDescDml::SetDataType(ENNETensorDataType DataType)
 		DML_TENSOR_DATA_TYPE DmlDataType = Util::GetTensorDataType(DataType);
 		check(DmlDataType != DML_TENSOR_DATA_TYPE_UNKNOWN);
 	
-		uint64 ElemByteSize = UE::NNECore::GetTensorDataTypeSizeInBytes(DataType);
+		uint64 ElemByteSize = UE::NNE::GetTensorDataTypeSizeInBytes(DataType);
 		check(ElemByteSize > 0)
 
 		BuffDesc.DataType = DmlDataType;
diff --git a/Engine/Plugins/Experimental/NNERuntimeRDG/Source/NNERuntimeRDG/Private/Dml/NNEDmlOperator.h b/Engine/Plugins/Experimental/NNERuntimeRDG/Source/NNERuntimeRDG/Private/Dml/NNEDmlOperator.h
--- a/Engine/Plugins/Experimental/NNERuntimeRDG/Source/NNERuntimeRDG/Private/Dml/NNEDmlOperator.h
+++ b/Engine/Plugins/Experimental/NNERuntimeRDG/Source/NNERuntimeRDG/Private/Dml/NNEDmlOperator.h
@@ -60,9 +60,9 @@ namespace Util
 {
 	
 template<typename T>
-using FSmallArray = TArray<T, TInlineAllocator<NNECore::FTensorShape::MaxRank>>;
-using FSmallIntArray = TArray<int32, TInlineAllocator<NNECore::FTensorShape::MaxRank>>;
-using FSmallUIntArray = TArray<uint32, TInlineAllocator<NNECore::FTensorShape::MaxRank>>;
+using FSmallArray = TArray<T, TInlineAllocator<NNE::FTensorShape::MaxRank>>;
+using FSmallIntArray = TArray<int32, TInlineAllocator<NNE::FTensorShape::MaxRank>>;
+using FSmallUIntArray = TArray<uint32, TInlineAllocator<NNE::FTensorShape::MaxRank>>;
 
 template<typename InputType, typename OutputType>
 inline bool IsOverflowing(InputType Input)
@@ -136,7 +136,7 @@ inline bool GetArrayAttributeNoOverflow(
 	return true;
 }
 
-extern bool IsSameShape(const NNECore::Internal::FTensor& Left, const NNECore::Internal::FTensor& Right);
+extern bool IsSameShape(const NNE::Internal::FTensor& Left, const NNE::Internal::FTensor& Right);
 extern bool IsSameShape(const FTensorDescDml& Left, const FTensorDescDml& Right);
 extern DML_TENSOR_DATA_TYPE GetTensorDataType(ENNETensorDataType DataType);
 
@@ -156,7 +156,7 @@ public:
 	FTensorDescDml& SetShape(TConstArrayView<uint32> Shape);
 
 	// Utility method that calls SetShape(Shape)
-	FTensorDescDml& SetShape(const UE::NNECore::FTensorShape& Shape)
+	FTensorDescDml& SetShape(const UE::NNE::FTensorShape& Shape)
 	{
 		return SetShape(Shape.GetData());
 	}
@@ -165,7 +165,7 @@ public:
 	FTensorDescDml& SetShape(TConstArrayView<uint32> Shape, int32 Rank);
 
 	// Utility method that calls SetShape(Shape, Rank)
-	FTensorDescDml& SetShape(const UE::NNECore::FTensorShape& Shape, int32 Rank)
+	FTensorDescDml& SetShape(const UE::NNE::FTensorShape& Shape, int32 Rank)
 	{
 		return SetShape(Shape.GetData(), Rank);
 	}
@@ -174,7 +174,7 @@ public:
 	FTensorDescDml& SetShape(TConstArrayView<uint32> Shape, TConstArrayView<uint32> BroadcastShape);
 
 	// Utility method that calles SetShape() from above
-	FTensorDescDml& SetShape(const UE::NNECore::FTensorShape& Shape, const UE::NNECore::FTensorShape& BroadcastShape)
+	FTensorDescDml& SetShape(const UE::NNE::FTensorShape& Shape, const UE::NNE::FTensorShape& BroadcastShape)
 	{
 		return SetShape(Shape.GetData(), BroadcastShape.GetData());
 	}
@@ -186,7 +186,7 @@ public:
 	FTensorDescDml& SetStridesFromShape(TConstArrayView<uint32> Shape);
 
 	// Set strides computed from shape (used for broadcasting)
-	FTensorDescDml& SetStridesFromShape(const UE::NNECore::FTensorShape& Shape)
+	FTensorDescDml& SetStridesFromShape(const UE::NNE::FTensorShape& Shape)
 	{
 		return SetStridesFromShape(Shape.GetData());
 	}
@@ -201,7 +201,7 @@ public:
 	FTensorDescDml& SetDataOwnedByDml(bool bSetOwnedByDml);
 
 	// Utility method to use FTensor to set all the members
-	FTensorDescDml& SetFromTensor(const UE::NNECore::Internal::FTensor& Tensor)
+	FTensorDescDml& SetFromTensor(const UE::NNE::Internal::FTensor& Tensor)
 	{
 		return SetShape(Tensor.GetShape())
 			.SetDataType(Tensor.GetDataType())
@@ -209,7 +209,7 @@ public:
 	}
 
 	// Utility method to use FTensor to set all the members
-	FTensorDescDml& SetFromTensor(const UE::NNECore::Internal::FTensor& Tensor, int32 Rank)
+	FTensorDescDml& SetFromTensor(const UE::NNE::Internal::FTensor& Tensor, int32 Rank)
 	{
 		return SetShape(Tensor.GetShape(), Rank)
 			.SetDataType(Tensor.GetDataType())
@@ -217,7 +217,7 @@ public:
 	}
 
 	// Utility method to use FTensor to set all the members
-	FTensorDescDml& SetFromTensorBroadcast(const UE::NNECore::Internal::FTensor& Tensor, const UE::NNECore::FTensorShape& BroadcastShape)
+	FTensorDescDml& SetFromTensorBroadcast(const UE::NNE::Internal::FTensor& Tensor, const UE::NNE::FTensorShape& BroadcastShape)
 	{
 		return SetShape(Tensor.GetShape(), BroadcastShape)
 			.SetDataType(Tensor.GetDataType())
@@ -225,7 +225,7 @@ public:
 	}
 
 	// Utility method to use FTensor to set all the members
-	FTensorDescDml& SetFromTensor1D(const UE::NNECore::Internal::FTensor& Tensor, int32 Rank)
+	FTensorDescDml& SetFromTensor1D(const UE::NNE::Internal::FTensor& Tensor, int32 Rank)
 	{
 		return SetShape1D(Tensor.GetShape().GetData()[0], Rank)
 			.SetDataType(Tensor.GetDataType())
@@ -291,7 +291,7 @@ public:
 
 	virtual ~FOperatorDml() = default;
 
-	virtual bool Initialize(IDMLDevice* Device, TArrayView<const NNECore::Internal::FTensor> InputTensors, TArrayView<const NNECore::Internal::FTensor> OutputTensors, const NNECore::FAttributeMap& Attributes) = 0;
+	virtual bool Initialize(IDMLDevice* Device, TArrayView<const NNE::Internal::FTensor> InputTensors, TArrayView<const NNE::Internal::FTensor> OutputTensors, const NNE::FAttributeMap& Attributes) = 0;
 
 	virtual TConstArrayView<int32> GetConstantCPUInputs() const;
 
diff --git a/Engine/Plugins/Experimental/NNERuntimeRDG/Source/NNERuntimeRDG/Private/Dml/NNEDmlOperatorActivations.cpp b/Engine/Plugins/Experimental/NNERuntimeRDG/Source/NNERuntimeRDG/Private/Dml/NNEDmlOperatorActivations.cpp
--- a/Engine/Plugins/Experimental/NNERuntimeRDG/Source/NNERuntimeRDG/Private/Dml/NNEDmlOperatorActivations.cpp
+++ b/Engine/Plugins/Experimental/NNERuntimeRDG/Source/NNERuntimeRDG/Private/Dml/NNEDmlOperatorActivations.cpp
@@ -23,7 +23,7 @@ public:
 		return new FOperatorDmlActivationUnary();
 	}
 
-	static bool Validate(const NNECore::FAttributeMap& AttributeMap, TConstArrayView<ENNETensorDataType> InputTypes, TConstArrayView<NNECore::FSymbolicTensorShape> InputShapes)
+	static bool Validate(const NNE::FAttributeMap& AttributeMap, TConstArrayView<ENNETensorDataType> InputTypes, TConstArrayView<NNE::FSymbolicTensorShape> InputShapes)
 	{
 		//TODO
 		return true;
@@ -46,10 +46,10 @@ public:
 	//
 	//
 	//
-	virtual bool Initialize(IDMLDevice* Device, TArrayView<const NNECore::Internal::FTensor> InputTensors, TArrayView<const NNECore::Internal::FTensor> OutputTensors, const NNECore::FAttributeMap& Attributes) override
+	virtual bool Initialize(IDMLDevice* Device, TArrayView<const NNE::Internal::FTensor> InputTensors, TArrayView<const NNE::Internal::FTensor> OutputTensors, const NNE::FAttributeMap& Attributes) override
 	{
-		const NNECore::Internal::FTensor& InputTensor = InputTensors[0];
-		const NNECore::Internal::FTensor& OutputTensor = OutputTensors[0];
+		const NNE::Internal::FTensor& InputTensor = InputTensors[0];
+		const NNE::Internal::FTensor& OutputTensor = OutputTensors[0];
 
 		if constexpr (std::is_same_v<DmlActivationOpDescType, DML_ACTIVATION_SOFTMAX1_OPERATOR_DESC> ||
 					  std::is_same_v<DmlActivationOpDescType, DML_ACTIVATION_LOG_SOFTMAX1_OPERATOR_DESC>)
@@ -187,7 +187,7 @@ public:
 		return new FOperatorDmlActivationBinary();
 	}
 
-	static bool Validate(const NNECore::FAttributeMap& AttributeMap, TConstArrayView<ENNETensorDataType> InputTypes, TConstArrayView<NNECore::FSymbolicTensorShape> InputShapes)
+	static bool Validate(const NNE::FAttributeMap& AttributeMap, TConstArrayView<ENNETensorDataType> InputTypes, TConstArrayView<NNE::FSymbolicTensorShape> InputShapes)
 	{
 		//TODO
 		return true;
@@ -202,11 +202,11 @@ public:
 	//
 	//
 	//
-	virtual bool Initialize(IDMLDevice* Device, TArrayView<const NNECore::Internal::FTensor> InputTensors, TArrayView<const NNECore::Internal::FTensor> OutputTensors, const NNECore::FAttributeMap& Attributes) override
+	virtual bool Initialize(IDMLDevice* Device, TArrayView<const NNE::Internal::FTensor> InputTensors, TArrayView<const NNE::Internal::FTensor> OutputTensors, const NNE::FAttributeMap& Attributes) override
 	{
-		const NNECore::Internal::FTensor& InputATensor = InputTensors[0];
-		const NNECore::Internal::FTensor& InputBTensor = InputTensors[1];
-		const NNECore::Internal::FTensor& OutputTensor = OutputTensors[0];
+		const NNE::Internal::FTensor& InputATensor = InputTensors[0];
+		const NNE::Internal::FTensor& InputBTensor = InputTensors[1];
+		const NNE::Internal::FTensor& OutputTensor = OutputTensors[0];
 
 		// Initialize tensor descriptors
 		FTensorDescDml	DmlInputATensorDesc;
diff --git a/Engine/Plugins/Experimental/NNERuntimeRDG/Source/NNERuntimeRDG/Private/Dml/NNEDmlOperatorBatchNormalization.cpp b/Engine/Plugins/Experimental/NNERuntimeRDG/Source/NNERuntimeRDG/Private/Dml/NNEDmlOperatorBatchNormalization.cpp
--- a/Engine/Plugins/Experimental/NNERuntimeRDG/Source/NNERuntimeRDG/Private/Dml/NNEDmlOperatorBatchNormalization.cpp
+++ b/Engine/Plugins/Experimental/NNERuntimeRDG/Source/NNERuntimeRDG/Private/Dml/NNEDmlOperatorBatchNormalization.cpp
@@ -27,7 +27,7 @@ public:
 		return new FOperatorDmlBatchNormalization();
 	}
 
-	static bool Validate(const NNECore::FAttributeMap& AttributeMap, TConstArrayView<ENNETensorDataType> InputTypes, TConstArrayView<NNECore::FSymbolicTensorShape> InputShapes)
+	static bool Validate(const NNE::FAttributeMap& AttributeMap, TConstArrayView<ENNETensorDataType> InputTypes, TConstArrayView<NNE::FSymbolicTensorShape> InputShapes)
 	{
 		//TODO
 		return true;
@@ -36,7 +36,7 @@ public:
 	//
 	//
 	//
-	virtual bool Initialize(IDMLDevice* Device, TArrayView<const NNECore::Internal::FTensor> InputTensors, TArrayView<const NNECore::Internal::FTensor> OutputTensors, const NNECore::FAttributeMap& Attributes) override
+	virtual bool Initialize(IDMLDevice* Device, TArrayView<const NNE::Internal::FTensor> InputTensors, TArrayView<const NNE::Internal::FTensor> OutputTensors, const NNE::FAttributeMap& Attributes) override
 	{
 		check(InputTensors.Num() == Count);
 		check(OutputTensors.Num() >= 1);
@@ -47,8 +47,8 @@ public:
 		RemappedInputs.Add(Scale);
 		RemappedInputs.Add(Bias);
 
-		const NNECore::Internal::FTensor& InputTensor = InputTensors[0];
-		const NNECore::Internal::FTensor& OutputTensor = OutputTensors[0];
+		const NNE::Internal::FTensor& InputTensor = InputTensors[0];
+		const NNE::Internal::FTensor& OutputTensor = OutputTensors[0];
 
 		if (InputTensor.GetShape().Rank() > 8)
 		{
@@ -78,7 +78,7 @@ public:
 
 		for (int32 Idx = Scale; Idx < Count; ++Idx)
 		{
-			const NNECore::Internal::FTensor& CurrTensor = InputTensors[Idx];
+			const NNE::Internal::FTensor& CurrTensor = InputTensors[Idx];
 			FTensorDescDml& DmlCurrTensorDesc = DmlInputTensorDescs[Idx];
 
 			if (CurrTensor.GetShape().Rank() == 1)
diff --git a/Engine/Plugins/Experimental/NNERuntimeRDG/Source/NNERuntimeRDG/Private/Dml/NNEDmlOperatorCast.cpp b/Engine/Plugins/Experimental/NNERuntimeRDG/Source/NNERuntimeRDG/Private/Dml/NNEDmlOperatorCast.cpp
--- a/Engine/Plugins/Experimental/NNERuntimeRDG/Source/NNERuntimeRDG/Private/Dml/NNEDmlOperatorCast.cpp
+++ b/Engine/Plugins/Experimental/NNERuntimeRDG/Source/NNERuntimeRDG/Private/Dml/NNEDmlOperatorCast.cpp
@@ -17,7 +17,7 @@ public:
 		return new FOperatorDmlCast();
 	}
 
-	static bool Validate(const NNECore::FAttributeMap& AttributeMap, TConstArrayView<ENNETensorDataType> InputTypes, TConstArrayView<NNECore::FSymbolicTensorShape> InputShapes)
+	static bool Validate(const NNE::FAttributeMap& AttributeMap, TConstArrayView<ENNETensorDataType> InputTypes, TConstArrayView<NNE::FSymbolicTensorShape> InputShapes)
 	{
 		//TODO
 		return true;
@@ -26,13 +26,13 @@ public:
 	//
 	//
 	//
-	virtual bool Initialize(IDMLDevice* Device, TArrayView<const NNECore::Internal::FTensor> InputTensors, TArrayView<const NNECore::Internal::FTensor> OutputTensors, const NNECore::FAttributeMap& Attributes) override
+	virtual bool Initialize(IDMLDevice* Device, TArrayView<const NNE::Internal::FTensor> InputTensors, TArrayView<const NNE::Internal::FTensor> OutputTensors, const NNE::FAttributeMap& Attributes) override
 	{
 		check(InputTensors.Num() == 1);
 		check(OutputTensors.Num() == 1);
 		
-		const NNECore::Internal::FTensor& InputTensor = InputTensors[0];
-		const NNECore::Internal::FTensor& OutputTensor = OutputTensors[0];
+		const NNE::Internal::FTensor& InputTensor = InputTensors[0];
+		const NNE::Internal::FTensor& OutputTensor = OutputTensors[0];
 
 		ENNETensorDataType To = (ENNETensorDataType) Attributes.GetValue<int32>(TEXT("to"));
 		
diff --git a/Engine/Plugins/Experimental/NNERuntimeRDG/Source/NNERuntimeRDG/Private/Dml/NNEDmlOperatorConcat.cpp b/Engine/Plugins/Experimental/NNERuntimeRDG/Source/NNERuntimeRDG/Private/Dml/NNEDmlOperatorConcat.cpp
--- a/Engine/Plugins/Experimental/NNERuntimeRDG/Source/NNERuntimeRDG/Private/Dml/NNEDmlOperatorConcat.cpp
+++ b/Engine/Plugins/Experimental/NNERuntimeRDG/Source/NNERuntimeRDG/Private/Dml/NNEDmlOperatorConcat.cpp
@@ -23,7 +23,7 @@ public:
 		return new FOperatorDmlConcat();
 	}
 
-	static bool Validate(const NNECore::FAttributeMap& AttributeMap, TConstArrayView<ENNETensorDataType> InputTypes, TConstArrayView<NNECore::FSymbolicTensorShape> InputShapes)
+	static bool Validate(const NNE::FAttributeMap& AttributeMap, TConstArrayView<ENNETensorDataType> InputTypes, TConstArrayView<NNE::FSymbolicTensorShape> InputShapes)
 	{
 		//TODO
 		return true;
@@ -32,7 +32,7 @@ public:
 	//
 	//
 	//
-	virtual bool Initialize(IDMLDevice* Device, TArrayView<const NNECore::Internal::FTensor> InputTensors, TArrayView<const NNECore::Internal::FTensor> OutputTensors, const NNECore::FAttributeMap& Attributes) override
+	virtual bool Initialize(IDMLDevice* Device, TArrayView<const NNE::Internal::FTensor> InputTensors, TArrayView<const NNE::Internal::FTensor> OutputTensors, const NNE::FAttributeMap& Attributes) override
 	{
 		check(InputTensors.Num() > 0);
 		check(OutputTensors.Num() == 1);
@@ -56,9 +56,9 @@ public:
 		TArray<DML_TENSOR_DESC> DmlInputDescs;
 		DmlInputDescs.SetNumUninitialized(InputTensors.Num());
 		
-		for (TConstEnumerateRef<NNECore::Internal::FTensor> It : EnumerateRange(InputTensors))
+		for (TConstEnumerateRef<NNE::Internal::FTensor> It : EnumerateRange(InputTensors))
 		{
-			const NNECore::Internal::FTensor& Tensor = *It;
+			const NNE::Internal::FTensor& Tensor = *It;
 			int Idx = It.GetIndex();
 			
 			if (Algo::Find(Tensor.GetShape().GetData(), 0) == nullptr)
diff --git a/Engine/Plugins/Experimental/NNERuntimeRDG/Source/NNERuntimeRDG/Private/Dml/NNEDmlOperatorConvolution.cpp b/Engine/Plugins/Experimental/NNERuntimeRDG/Source/NNERuntimeRDG/Private/Dml/NNEDmlOperatorConvolution.cpp
--- a/Engine/Plugins/Experimental/NNERuntimeRDG/Source/NNERuntimeRDG/Private/Dml/NNEDmlOperatorConvolution.cpp
+++ b/Engine/Plugins/Experimental/NNERuntimeRDG/Source/NNERuntimeRDG/Private/Dml/NNEDmlOperatorConvolution.cpp
@@ -34,7 +34,7 @@ class FOperatorDmlConv : public FOperatorDml
 		//
 		//
 		//
-		bool Init(const NNECore::FTensorShape& InputShape, const NNECore::FTensorShape& FilterShape, const NNECore::FAttributeMap& Attributes)
+		bool Init(const NNE::FTensorShape& InputShape, const NNE::FTensorShape& FilterShape, const NNE::FAttributeMap& Attributes)
 		{
 			check(InputShape.Rank() > NonspatialDimensionCount);
 			check(FilterShape.Rank() == InputShape.Rank());
@@ -207,7 +207,7 @@ public:
 		return new FOperatorDmlConv();
 	}
 
-	static bool Validate(const NNECore::FAttributeMap& AttributeMap, TConstArrayView<ENNETensorDataType> InputTypes, TConstArrayView<NNECore::FSymbolicTensorShape> InputShapes)
+	static bool Validate(const NNE::FAttributeMap& AttributeMap, TConstArrayView<ENNETensorDataType> InputTypes, TConstArrayView<NNE::FSymbolicTensorShape> InputShapes)
 	{
 		//TODO
 		return true;
@@ -216,10 +216,10 @@ public:
 	//
 	//
 	//
-	virtual bool Initialize(IDMLDevice* Device, TArrayView<const NNECore::Internal::FTensor> InputTensors, TArrayView<const NNECore::Internal::FTensor> OutputTensors, const NNECore::FAttributeMap& Attributes) override
+	virtual bool Initialize(IDMLDevice* Device, TArrayView<const NNE::Internal::FTensor> InputTensors, TArrayView<const NNE::Internal::FTensor> OutputTensors, const NNE::FAttributeMap& Attributes) override
 	{
-		const NNECore::Internal::FTensor& InputTensor = InputTensors[0];
-		const NNECore::Internal::FTensor& FilterTensor = InputTensors[1];
+		const NNE::Internal::FTensor& InputTensor = InputTensors[0];
+		const NNE::Internal::FTensor& FilterTensor = InputTensors[1];
 		
 		FConvArgs	Args;
 		
@@ -228,11 +228,11 @@ public:
 			return false;
 		}
 
-		NNECore::Internal::FTensor OutputTensor = OutputTensors[0];
+		NNE::Internal::FTensor OutputTensor = OutputTensors[0];
 
 		if (Direction == DML_CONVOLUTION_DIRECTION_FORWARD)
 		{
-			OutputTensor.SetShape(NNECore::FTensorShape::Make(Args.OutputShape));
+			OutputTensor.SetShape(NNE::FTensorShape::Make(Args.OutputShape));
 		}
 
 		// Initialize tensor descriptors
@@ -261,7 +261,7 @@ public:
 
 		if (InputTensors.Num() > 2)
 		{
-			const NNECore::Internal::FTensor& BiasTensor = InputTensors[2];
+			const NNE::Internal::FTensor& BiasTensor = InputTensors[2];
 
 			if (!DmlBiasTensorDesc
 					.SetTensorRank(3, 5)
diff --git a/Engine/Plugins/Experimental/NNERuntimeRDG/Source/NNERuntimeRDG/Private/Dml/NNEDmlOperatorCumSum.cpp b/Engine/Plugins/Experimental/NNERuntimeRDG/Source/NNERuntimeRDG/Private/Dml/NNEDmlOperatorCumSum.cpp
--- a/Engine/Plugins/Experimental/NNERuntimeRDG/Source/NNERuntimeRDG/Private/Dml/NNEDmlOperatorCumSum.cpp
+++ b/Engine/Plugins/Experimental/NNERuntimeRDG/Source/NNERuntimeRDG/Private/Dml/NNEDmlOperatorCumSum.cpp
@@ -18,7 +18,7 @@ public:
 		return new FOperatorDmlCumSum();
 	}
 
-	static bool Validate(const NNECore::FAttributeMap& AttributeMap, TConstArrayView<ENNETensorDataType> InputTypes, TConstArrayView<NNECore::FSymbolicTensorShape> InputShapes)
+	static bool Validate(const NNE::FAttributeMap& AttributeMap, TConstArrayView<ENNETensorDataType> InputTypes, TConstArrayView<NNE::FSymbolicTensorShape> InputShapes)
 	{
 		//TODO
 		return true;
@@ -27,7 +27,7 @@ public:
 	//
 	//
 	//
-	virtual bool Initialize(IDMLDevice* Device, TArrayView<const NNECore::Internal::FTensor> InputTensors, TArrayView<const NNECore::Internal::FTensor> OutputTensors, const NNECore::FAttributeMap& Attributes) override
+	virtual bool Initialize(IDMLDevice* Device, TArrayView<const NNE::Internal::FTensor> InputTensors, TArrayView<const NNE::Internal::FTensor> OutputTensors, const NNE::FAttributeMap& Attributes) override
 	{
 		check(InputTensors.Num() == 2);
 		check(OutputTensors.Num() == 1);
@@ -76,8 +76,8 @@ public:
 			return false;
         };
 
-		const NNECore::Internal::FTensor& InputTensor = InputTensors[0];
-		const NNECore::Internal::FTensor& OutputTensor = OutputTensors[0];
+		const NNE::Internal::FTensor& InputTensor = InputTensors[0];
+		const NNE::Internal::FTensor& OutputTensor = OutputTensors[0];
 
 		FTensorDescDml	DmlInputTensorDesc;
 		FTensorDescDml	DmlOutputTensorDesc;
diff --git a/Engine/Plugins/Experimental/NNERuntimeRDG/Source/NNERuntimeRDG/Private/Dml/NNEDmlOperatorDepthToSpace.cpp b/Engine/Plugins/Experimental/NNERuntimeRDG/Source/NNERuntimeRDG/Private/Dml/NNEDmlOperatorDepthToSpace.cpp
--- a/Engine/Plugins/Experimental/NNERuntimeRDG/Source/NNERuntimeRDG/Private/Dml/NNEDmlOperatorDepthToSpace.cpp
+++ b/Engine/Plugins/Experimental/NNERuntimeRDG/Source/NNERuntimeRDG/Private/Dml/NNEDmlOperatorDepthToSpace.cpp
@@ -30,7 +30,7 @@ public:
 		return new FOperatorDmlDepthToSpace();
 	}
 
-	static bool Validate(const NNECore::FAttributeMap& AttributeMap, TConstArrayView<ENNETensorDataType> InputTypes, TConstArrayView<NNECore::FSymbolicTensorShape> InputShapes)
+	static bool Validate(const NNE::FAttributeMap& AttributeMap, TConstArrayView<ENNETensorDataType> InputTypes, TConstArrayView<NNE::FSymbolicTensorShape> InputShapes)
 	{
 		//TODO
 		return true;
@@ -39,7 +39,7 @@ public:
 	//
 	//
 	//
-	virtual bool Initialize(IDMLDevice* Device, TArrayView<const NNECore::Internal::FTensor> InputTensors, TArrayView<const NNECore::Internal::FTensor> OutputTensors, const NNECore::FAttributeMap& Attributes) override
+	virtual bool Initialize(IDMLDevice* Device, TArrayView<const NNE::Internal::FTensor> InputTensors, TArrayView<const NNE::Internal::FTensor> OutputTensors, const NNE::FAttributeMap& Attributes) override
 	{
 		check(InputTensors.Num() == 1);
 		check(OutputTensors.Num() == 1);
@@ -58,8 +58,8 @@ public:
 		}
 
 
-		const NNECore::Internal::FTensor& InputTensorDesc = InputTensors[0];
-		const NNECore::Internal::FTensor& OutputTensorDesc = OutputTensors[0];
+		const NNE::Internal::FTensor& InputTensorDesc = InputTensors[0];
+		const NNE::Internal::FTensor& OutputTensorDesc = OutputTensors[0];
 
 		FTensorDescDml	DmlInputTensorDesc;
 		FTensorDescDml	DmlOutputTensorDesc;
diff --git a/Engine/Plugins/Experimental/NNERuntimeRDG/Source/NNERuntimeRDG/Private/Dml/NNEDmlOperatorElementwise.cpp b/Engine/Plugins/Experimental/NNERuntimeRDG/Source/NNERuntimeRDG/Private/Dml/NNEDmlOperatorElementwise.cpp
--- a/Engine/Plugins/Experimental/NNERuntimeRDG/Source/NNERuntimeRDG/Private/Dml/NNEDmlOperatorElementwise.cpp
+++ b/Engine/Plugins/Experimental/NNERuntimeRDG/Source/NNERuntimeRDG/Private/Dml/NNEDmlOperatorElementwise.cpp
@@ -23,7 +23,7 @@ public:
 		return new FOperatorDmlElementWiseUnary();
 	}
 
-	static bool Validate(const NNECore::FAttributeMap& AttributeMap, TConstArrayView<ENNETensorDataType> InputTypes, TConstArrayView<NNECore::FSymbolicTensorShape> InputShapes)
+	static bool Validate(const NNE::FAttributeMap& AttributeMap, TConstArrayView<ENNETensorDataType> InputTypes, TConstArrayView<NNE::FSymbolicTensorShape> InputShapes)
 	{
 		//TODO
 		return true;
@@ -42,11 +42,11 @@ public:
 	//
 	//
 	//
-	virtual bool Initialize(IDMLDevice* Device, TArrayView<const NNECore::Internal::FTensor> InputTensors, TArrayView<const NNECore::Internal::FTensor> OutputTensors, const NNECore::FAttributeMap& Attributes) override
+	virtual bool Initialize(IDMLDevice* Device, TArrayView<const NNE::Internal::FTensor> InputTensors, TArrayView<const NNE::Internal::FTensor> OutputTensors, const NNE::FAttributeMap& Attributes) override
 	{
 
-		const NNECore::Internal::FTensor& InputTensor = InputTensors[0];
-		const NNECore::Internal::FTensor& OutputTensor = OutputTensors[0];
+		const NNE::Internal::FTensor& InputTensor = InputTensors[0];
+		const NNE::Internal::FTensor& OutputTensor = OutputTensors[0];
 
 		if constexpr (std::is_same_v<DmlElementWiseOpDescType, DML_ELEMENT_WISE_CLIP_OPERATOR_DESC>)
 		{
@@ -142,11 +142,11 @@ public:
 	//
 	//
 	//
-	virtual bool Initialize(IDMLDevice* Device, TArrayView<const NNECore::Internal::FTensor> InputTensors, TArrayView<const NNECore::Internal::FTensor> OutputTensors, const NNECore::FAttributeMap& Attributes) override
+	virtual bool Initialize(IDMLDevice* Device, TArrayView<const NNE::Internal::FTensor> InputTensors, TArrayView<const NNE::Internal::FTensor> OutputTensors, const NNE::FAttributeMap& Attributes) override
 	{
-		const NNECore::Internal::FTensor& InputATensor = InputTensors[0];
-		const NNECore::Internal::FTensor& InputBTensor = InputTensors[1];
-		const NNECore::Internal::FTensor& OutputTensor = OutputTensors[0];
+		const NNE::Internal::FTensor& InputATensor = InputTensors[0];
+		const NNE::Internal::FTensor& InputBTensor = InputTensors[1];
+		const NNE::Internal::FTensor& OutputTensor = OutputTensors[0];
 
 		// Initialize tensor descriptors
 		FTensorDescDml	DmlInputATensorDesc;
diff --git a/Engine/Plugins/Experimental/NNERuntimeRDG/Source/NNERuntimeRDG/Private/Dml/NNEDmlOperatorFlatten.cpp b/Engine/Plugins/Experimental/NNERuntimeRDG/Source/NNERuntimeRDG/Private/Dml/NNEDmlOperatorFlatten.cpp
--- a/Engine/Plugins/Experimental/NNERuntimeRDG/Source/NNERuntimeRDG/Private/Dml/NNEDmlOperatorFlatten.cpp
+++ b/Engine/Plugins/Experimental/NNERuntimeRDG/Source/NNERuntimeRDG/Private/Dml/NNEDmlOperatorFlatten.cpp
@@ -18,7 +18,7 @@ public:
 		return new FOperatorDmlFlatten();
 	}
 
-	static bool Validate(const NNECore::FAttributeMap& AttributeMap, TConstArrayView<ENNETensorDataType> InputTypes, TConstArrayView<NNECore::FSymbolicTensorShape> InputShapes)
+	static bool Validate(const NNE::FAttributeMap& AttributeMap, TConstArrayView<ENNETensorDataType> InputTypes, TConstArrayView<NNE::FSymbolicTensorShape> InputShapes)
 	{
 		//TODO
 		return true;
@@ -27,7 +27,7 @@ public:
 	//
 	//
 	//
-	virtual bool Initialize(IDMLDevice* Device, TArrayView<const NNECore::Internal::FTensor> InputTensors, TArrayView<const NNECore::Internal::FTensor> OutputTensors, const NNECore::FAttributeMap& Attributes) override
+	virtual bool Initialize(IDMLDevice* Device, TArrayView<const NNE::Internal::FTensor> InputTensors, TArrayView<const NNE::Internal::FTensor> OutputTensors, const NNE::FAttributeMap& Attributes) override
 	{
 		check(InputTensors.Num() == 1);
 		check(OutputTensors.Num() == 1);
diff --git a/Engine/Plugins/Experimental/NNERuntimeRDG/Source/NNERuntimeRDG/Private/Dml/NNEDmlOperatorGRU.cpp b/Engine/Plugins/Experimental/NNERuntimeRDG/Source/NNERuntimeRDG/Private/Dml/NNEDmlOperatorGRU.cpp
--- a/Engine/Plugins/Experimental/NNERuntimeRDG/Source/NNERuntimeRDG/Private/Dml/NNEDmlOperatorGRU.cpp
+++ b/Engine/Plugins/Experimental/NNERuntimeRDG/Source/NNERuntimeRDG/Private/Dml/NNEDmlOperatorGRU.cpp
@@ -125,7 +125,7 @@ public:
 		return new FOperatorDmlGRU();
 	}
 
-	static bool Validate(const NNECore::FAttributeMap& AttributeMap, TConstArrayView<ENNETensorDataType> InputTypes, TConstArrayView<NNECore::FSymbolicTensorShape> InputShapes)
+	static bool Validate(const NNE::FAttributeMap& AttributeMap, TConstArrayView<ENNETensorDataType> InputTypes, TConstArrayView<NNE::FSymbolicTensorShape> InputShapes)
 	{
 		//TODO
 		return true;
@@ -134,7 +134,7 @@ public:
 	//
 	//
 	//
-	virtual bool Initialize(IDMLDevice* Device, TArrayView<const NNECore::Internal::FTensor> InputTensors, TArrayView<const NNECore::Internal::FTensor> OutputTensors, const NNECore::FAttributeMap& Attributes) override
+	virtual bool Initialize(IDMLDevice* Device, TArrayView<const NNE::Internal::FTensor> InputTensors, TArrayView<const NNE::Internal::FTensor> OutputTensors, const NNE::FAttributeMap& Attributes) override
 	{
 		
 #define SET_DMLTENSORDESC_FROM_TENSOR_COND_DATATYPE(OpDesc, Tensor, Name, DataType, Cond)\
diff --git a/Engine/Plugins/Experimental/NNERuntimeRDG/Source/NNERuntimeRDG/Private/Dml/NNEDmlOperatorGather.cpp b/Engine/Plugins/Experimental/NNERuntimeRDG/Source/NNERuntimeRDG/Private/Dml/NNEDmlOperatorGather.cpp
--- a/Engine/Plugins/Experimental/NNERuntimeRDG/Source/NNERuntimeRDG/Private/Dml/NNEDmlOperatorGather.cpp
+++ b/Engine/Plugins/Experimental/NNERuntimeRDG/Source/NNERuntimeRDG/Private/Dml/NNEDmlOperatorGather.cpp
@@ -16,7 +16,7 @@ public:
 		return new FOperatorDmlGather();
 	}
 
-	static bool Validate(const NNECore::FAttributeMap& AttributeMap, TConstArrayView<ENNETensorDataType> InputTypes, TConstArrayView<NNECore::FSymbolicTensorShape> InputShapes)
+	static bool Validate(const NNE::FAttributeMap& AttributeMap, TConstArrayView<ENNETensorDataType> InputTypes, TConstArrayView<NNE::FSymbolicTensorShape> InputShapes)
 	{
 		//TODO
 		return true;
@@ -25,17 +25,17 @@ public:
 	//
 	//
 	//
-	virtual bool Initialize(IDMLDevice* Device, TArrayView<const NNECore::Internal::FTensor> InputTensors, TArrayView<const NNECore::Internal::FTensor> OutputTensors, const NNECore::FAttributeMap& Attributes) override
+	virtual bool Initialize(IDMLDevice* Device, TArrayView<const NNE::Internal::FTensor> InputTensors, TArrayView<const NNE::Internal::FTensor> OutputTensors, const NNE::FAttributeMap& Attributes) override
 	{
 		check(InputTensors.Num() == 2);
 		check(OutputTensors.Num() == 1);
 		
-		const NNECore::Internal::FTensor& InputTensor = InputTensors[0];
-		const NNECore::Internal::FTensor& IndicesTensor = InputTensors[1];
-		const NNECore::Internal::FTensor& OutputTensor = OutputTensors[0];
+		const NNE::Internal::FTensor& InputTensor = InputTensors[0];
+		const NNE::Internal::FTensor& IndicesTensor = InputTensors[1];
+		const NNE::Internal::FTensor& OutputTensor = OutputTensors[0];
 
-		const NNECore::FTensorShape& InputShape = InputTensor.GetShape();
-		const NNECore::FTensorShape& IndicesShape = IndicesTensor.GetShape();
+		const NNE::FTensorShape& InputShape = InputTensor.GetShape();
+		const NNE::FTensorShape& IndicesShape = IndicesTensor.GetShape();
 
 		if (IndicesShape.Rank() > InputShape.Rank())
 		{
diff --git a/Engine/Plugins/Experimental/NNERuntimeRDG/Source/NNERuntimeRDG/Private/Dml/NNEDmlOperatorGemm.cpp b/Engine/Plugins/Experimental/NNERuntimeRDG/Source/NNERuntimeRDG/Private/Dml/NNEDmlOperatorGemm.cpp
--- a/Engine/Plugins/Experimental/NNERuntimeRDG/Source/NNERuntimeRDG/Private/Dml/NNEDmlOperatorGemm.cpp
+++ b/Engine/Plugins/Experimental/NNERuntimeRDG/Source/NNERuntimeRDG/Private/Dml/NNEDmlOperatorGemm.cpp
@@ -18,7 +18,7 @@ public:
 		return new FOperatorDmlGemm();
 	}
 
-	static bool Validate(const NNECore::FAttributeMap& AttributeMap, TConstArrayView<ENNETensorDataType> InputTypes, TConstArrayView<NNECore::FSymbolicTensorShape> InputShapes)
+	static bool Validate(const NNE::FAttributeMap& AttributeMap, TConstArrayView<ENNETensorDataType> InputTypes, TConstArrayView<NNE::FSymbolicTensorShape> InputShapes)
 	{
 		//TODO
 		return true;
@@ -27,7 +27,7 @@ public:
 	//
 	//
 	//
-	virtual bool Initialize(IDMLDevice* Device, TArrayView<const NNECore::Internal::FTensor> InputTensors, TArrayView<const NNECore::Internal::FTensor> OutputTensors, const NNECore::FAttributeMap& Attributes) override
+	virtual bool Initialize(IDMLDevice* Device, TArrayView<const NNE::Internal::FTensor> InputTensors, TArrayView<const NNE::Internal::FTensor> OutputTensors, const NNE::FAttributeMap& Attributes) override
 	{
 		// Setup attributes
 		float Alpha = 1.0f;
@@ -40,9 +40,9 @@ public:
 		TransA = Attributes.GetValueOrDefault(TEXT("transA"), TransA);
 		TransB = Attributes.GetValueOrDefault(TEXT("transB"), TransB);
 		
-		const NNECore::Internal::FTensor& InputATensor = InputTensors[0];
-		const NNECore::Internal::FTensor& InputBTensor = InputTensors[1];
-		const NNECore::Internal::FTensor& OutputTensor = OutputTensors[0];
+		const NNE::Internal::FTensor& InputATensor = InputTensors[0];
+		const NNE::Internal::FTensor& InputBTensor = InputTensors[1];
+		const NNE::Internal::FTensor& OutputTensor = OutputTensors[0];
 
 		// Initialize tensor descriptors
 		FTensorDescDml	DmlInputATensorDesc;
@@ -70,7 +70,7 @@ public:
 
 		if (InputTensors.Num() > 2)
 		{
-			const NNECore::Internal::FTensor& InputCTensor = InputTensors[2];
+			const NNE::Internal::FTensor& InputCTensor = InputTensors[2];
 
 			if (!DmlInputCTensorDesc
 					.SetTensorRank(2, 4)
diff --git a/Engine/Plugins/Experimental/NNERuntimeRDG/Source/NNERuntimeRDG/Private/Dml/NNEDmlOperatorInstanceNormalization.cpp b/Engine/Plugins/Experimental/NNERuntimeRDG/Source/NNERuntimeRDG/Private/Dml/NNEDmlOperatorInstanceNormalization.cpp
--- a/Engine/Plugins/Experimental/NNERuntimeRDG/Source/NNERuntimeRDG/Private/Dml/NNEDmlOperatorInstanceNormalization.cpp
+++ b/Engine/Plugins/Experimental/NNERuntimeRDG/Source/NNERuntimeRDG/Private/Dml/NNEDmlOperatorInstanceNormalization.cpp
@@ -17,7 +17,7 @@ public:
 		return new FOperatorDmlInstanceNormalization();
 	}
 
-	static bool Validate(const NNECore::FAttributeMap& AttributeMap, TConstArrayView<ENNETensorDataType> InputTypes, TConstArrayView<NNECore::FSymbolicTensorShape> InputShapes)
+	static bool Validate(const NNE::FAttributeMap& AttributeMap, TConstArrayView<ENNETensorDataType> InputTypes, TConstArrayView<NNE::FSymbolicTensorShape> InputShapes)
 	{
 		//TODO
 		return true;
@@ -26,13 +26,13 @@ public:
 	//
 	//
 	//
-	virtual bool Initialize(IDMLDevice* Device, TArrayView<const NNECore::Internal::FTensor> InputTensors, TArrayView<const NNECore::Internal::FTensor> OutputTensors, const NNECore::FAttributeMap& Attributes) override
+	virtual bool Initialize(IDMLDevice* Device, TArrayView<const NNE::Internal::FTensor> InputTensors, TArrayView<const NNE::Internal::FTensor> OutputTensors, const NNE::FAttributeMap& Attributes) override
 	{
 		check(InputTensors.Num() >= 1 && InputTensors.Num() <= 3);
 		check(OutputTensors.Num() == 1);
 
-		const NNECore::Internal::FTensor& InputTensor = InputTensors[0];
-		const NNECore::Internal::FTensor& OutputTensor = OutputTensors[0];
+		const NNE::Internal::FTensor& InputTensor = InputTensors[0];
+		const NNE::Internal::FTensor& OutputTensor = OutputTensors[0];
 
 		if (InputTensor.GetShape().Rank() > 8)
 		{
@@ -61,7 +61,7 @@ public:
 
 		if (InputTensors.Num() > 1)
 		{
-			const NNECore::Internal::FTensor& ScaleTensor = InputTensors[1];
+			const NNE::Internal::FTensor& ScaleTensor = InputTensors[1];
 
 			if (!DmlScalingTensorDesc
 					.SetTensorRank(4, 4)
@@ -75,7 +75,7 @@ public:
 
 		if (InputTensors.Num() > 2)
 		{
-			const NNECore::Internal::FTensor& BiasTensor = InputTensors[2];
+			const NNE::Internal::FTensor& BiasTensor = InputTensors[2];
 
 			if (!DmlBiasTensorDesc
 					.SetTensorRank(4, 4)
diff --git a/Engine/Plugins/Experimental/NNERuntimeRDG/Source/NNERuntimeRDG/Private/Dml/NNEDmlOperatorMatMul.cpp b/Engine/Plugins/Experimental/NNERuntimeRDG/Source/NNERuntimeRDG/Private/Dml/NNEDmlOperatorMatMul.cpp
--- a/Engine/Plugins/Experimental/NNERuntimeRDG/Source/NNERuntimeRDG/Private/Dml/NNEDmlOperatorMatMul.cpp
+++ b/Engine/Plugins/Experimental/NNERuntimeRDG/Source/NNERuntimeRDG/Private/Dml/NNEDmlOperatorMatMul.cpp
@@ -18,7 +18,7 @@ public:
 		return new FOperatorDmlMatMul();
 	}
 
-	static bool Validate(const NNECore::FAttributeMap& AttributeMap, TConstArrayView<ENNETensorDataType> InputTypes, TConstArrayView<NNECore::FSymbolicTensorShape> InputShapes)
+	static bool Validate(const NNE::FAttributeMap& AttributeMap, TConstArrayView<ENNETensorDataType> InputTypes, TConstArrayView<NNE::FSymbolicTensorShape> InputShapes)
 	{
 		//TODO
 		return true;
@@ -27,16 +27,16 @@ public:
 	//
 	//
 	//
-	virtual bool Initialize(IDMLDevice* Device, TArrayView<const NNECore::Internal::FTensor> InputTensors, TArrayView<const NNECore::Internal::FTensor> OutputTensors, const NNECore::FAttributeMap& Attributes) override
+	virtual bool Initialize(IDMLDevice* Device, TArrayView<const NNE::Internal::FTensor> InputTensors, TArrayView<const NNE::Internal::FTensor> OutputTensors, const NNE::FAttributeMap& Attributes) override
 	{
 		check(InputTensors.Num() == 2);
 		check(OutputTensors.Num() == 1);
 
-		using namespace UE::NNECore;
+		using namespace UE::NNE;
 
-		const NNECore::Internal::FTensor& ATensor = InputTensors[0];
-		const NNECore::Internal::FTensor& BTensor = InputTensors[1];
-		const NNECore::Internal::FTensor& OutTensor = OutputTensors[0];
+		const NNE::Internal::FTensor& ATensor = InputTensors[0];
+		const NNE::Internal::FTensor& BTensor = InputTensors[1];
+		const NNE::Internal::FTensor& OutTensor = OutputTensors[0];
 
 		auto TrimMatrixStackLeadingOnes = [](TConstArrayView<uint32>&& Input) -> TConstArrayView<uint32>
 		{
diff --git a/Engine/Plugins/Experimental/NNERuntimeRDG/Source/NNERuntimeRDG/Private/Dml/NNEDmlOperatorMaxUnpool.cpp b/Engine/Plugins/Experimental/NNERuntimeRDG/Source/NNERuntimeRDG/Private/Dml/NNEDmlOperatorMaxUnpool.cpp
--- a/Engine/Plugins/Experimental/NNERuntimeRDG/Source/NNERuntimeRDG/Private/Dml/NNEDmlOperatorMaxUnpool.cpp
+++ b/Engine/Plugins/Experimental/NNERuntimeRDG/Source/NNERuntimeRDG/Private/Dml/NNEDmlOperatorMaxUnpool.cpp
@@ -18,7 +18,7 @@ public:
 		return new FOperatorDmlMaxUnpool();
 	}
 
-	static bool Validate(const NNECore::FAttributeMap& AttributeMap, TConstArrayView<ENNETensorDataType> InputTypes, TConstArrayView<NNECore::FSymbolicTensorShape> InputShapes)
+	static bool Validate(const NNE::FAttributeMap& AttributeMap, TConstArrayView<ENNETensorDataType> InputTypes, TConstArrayView<NNE::FSymbolicTensorShape> InputShapes)
 	{
 		//TODO
 		return true;
@@ -27,7 +27,7 @@ public:
 	//
 	//
 	//
-	virtual bool Initialize(IDMLDevice* Device, TArrayView<const NNECore::Internal::FTensor> InputTensors, TArrayView<const NNECore::Internal::FTensor> OutputTensors, const NNECore::FAttributeMap& Attributes) override
+	virtual bool Initialize(IDMLDevice* Device, TArrayView<const NNE::Internal::FTensor> InputTensors, TArrayView<const NNE::Internal::FTensor> OutputTensors, const NNE::FAttributeMap& Attributes) override
 	{
 		check(InputTensors.Num() == 2 || InputTensors.Num() == 3);
 		check(OutputTensors.Num() == 1);
@@ -37,9 +37,9 @@ public:
 			ConstantCPUInputs.Add(2);
 		}
 
-		const NNECore::Internal::FTensor& InputTensor = InputTensors[0];
-		const NNECore::Internal::FTensor& IndicesTensor = InputTensors[1];
-		const NNECore::Internal::FTensor& OutputTensor = OutputTensors[0];
+		const NNE::Internal::FTensor& InputTensor = InputTensors[0];
+		const NNE::Internal::FTensor& IndicesTensor = InputTensors[1];
+		const NNE::Internal::FTensor& OutputTensor = OutputTensors[0];
 
 		FTensorDescDml	DmlInputTensorDesc;
 		FTensorDescDml	DmlIndicesTensorDesc;
diff --git a/Engine/Plugins/Experimental/NNERuntimeRDG/Source/NNERuntimeRDG/Private/Dml/NNEDmlOperatorResample.cpp b/Engine/Plugins/Experimental/NNERuntimeRDG/Source/NNERuntimeRDG/Private/Dml/NNEDmlOperatorResample.cpp
--- a/Engine/Plugins/Experimental/NNERuntimeRDG/Source/NNERuntimeRDG/Private/Dml/NNEDmlOperatorResample.cpp
+++ b/Engine/Plugins/Experimental/NNERuntimeRDG/Source/NNERuntimeRDG/Private/Dml/NNEDmlOperatorResample.cpp
@@ -56,7 +56,7 @@ public:
 		return new FOperatorDmlResample<IsResize>();
 	}
 
-	static bool Validate(const NNECore::FAttributeMap& AttributeMap, TConstArrayView<ENNETensorDataType> InputTypes, TConstArrayView<NNECore::FSymbolicTensorShape> InputShapes)
+	static bool Validate(const NNE::FAttributeMap& AttributeMap, TConstArrayView<ENNETensorDataType> InputTypes, TConstArrayView<NNE::FSymbolicTensorShape> InputShapes)
 	{
 		//TODO
 		return true;
@@ -65,14 +65,14 @@ public:
 	//
 	//
 	//
-	virtual bool Initialize(IDMLDevice* Device, TArrayView<const NNECore::Internal::FTensor> InputTensors, TArrayView<const NNECore::Internal::FTensor> OutputTensors, const NNECore::FAttributeMap& Attributes) override
+	virtual bool Initialize(IDMLDevice* Device, TArrayView<const NNE::Internal::FTensor> InputTensors, TArrayView<const NNE::Internal::FTensor> OutputTensors, const NNE::FAttributeMap& Attributes) override
 	{
 		check(InputTensors.Num() >= 1 && InputTensors.Num() < 4);
 		check(OutputTensors.Num() == 1);
 
-		const NNECore::Internal::FTensor& InputTensor = InputTensors[0];
-		const NNECore::Internal::FTensor& ScaleTensor = (InputTensors.Num() == 2)?InputTensors[1]:InputTensors[2];//Upsample have scale at position 1, while Resize at 2.
-		const NNECore::Internal::FTensor& OutputTensor = OutputTensors[0];
+		const NNE::Internal::FTensor& InputTensor = InputTensors[0];
+		const NNE::Internal::FTensor& ScaleTensor = (InputTensors.Num() == 2)?InputTensors[1]:InputTensors[2];//Upsample have scale at position 1, while Resize at 2.
+		const NNE::Internal::FTensor& OutputTensor = OutputTensors[0];
 
 		if (InputTensors.Num() == 3 && InputTensors[2].GetDataType() == ENNETensorDataType::Int64)
 		{
@@ -139,7 +139,7 @@ public:
 				{
 					// NOTE: no tests for this, ORT erroneously puts all 0.0fs in the output tensor in this case.
 					
-					const NNECore::Internal::FTensor& RoiTensor = InputTensors[1];
+					const NNE::Internal::FTensor& RoiTensor = InputTensors[1];
 					if (!RoiTensor.HasPreparedData())
 					{
 						UE_LOG(LogNNE, Warning, TEXT("roi should be a constant tensor, it is here a variable tensor of name %s."), *RoiTensor.GetName());
diff --git a/Engine/Plugins/Experimental/NNERuntimeRDG/Source/NNERuntimeRDG/Private/Dml/NNEDmlOperatorSlice.cpp b/Engine/Plugins/Experimental/NNERuntimeRDG/Source/NNERuntimeRDG/Private/Dml/NNEDmlOperatorSlice.cpp
--- a/Engine/Plugins/Experimental/NNERuntimeRDG/Source/NNERuntimeRDG/Private/Dml/NNEDmlOperatorSlice.cpp
+++ b/Engine/Plugins/Experimental/NNERuntimeRDG/Source/NNERuntimeRDG/Private/Dml/NNEDmlOperatorSlice.cpp
@@ -19,7 +19,7 @@ class FOperatorDmlSlice : public FOperatorDml
 {	
 	template<typename DataType>
 	void ComputeOffsetsSizesStrides(
-		TArrayView<const NNECore::Internal::FTensor> InputTensors,
+		TArrayView<const NNE::Internal::FTensor> InputTensors,
 		Util::FSmallUIntArray& OutOffsets, 
 		Util::FSmallUIntArray& OutSizes, 
 		Util::FSmallIntArray& OutStrides)
@@ -114,7 +114,7 @@ public:
 		return new FOperatorDmlSlice();
 	}
 
-	static bool Validate(const NNECore::FAttributeMap& AttributeMap, TConstArrayView<ENNETensorDataType> InputTypes, TConstArrayView<NNECore::FSymbolicTensorShape> InputShapes)
+	static bool Validate(const NNE::FAttributeMap& AttributeMap, TConstArrayView<ENNETensorDataType> InputTypes, TConstArrayView<NNE::FSymbolicTensorShape> InputShapes)
 	{
 		//TODO
 		return true;
@@ -123,7 +123,7 @@ public:
 	//
 	//
 	//
-	virtual bool Initialize(IDMLDevice* Device, TArrayView<const NNECore::Internal::FTensor> InputTensors, TArrayView<const NNECore::Internal::FTensor> OutputTensors, const NNECore::FAttributeMap& Attributes) override
+	virtual bool Initialize(IDMLDevice* Device, TArrayView<const NNE::Internal::FTensor> InputTensors, TArrayView<const NNE::Internal::FTensor> OutputTensors, const NNE::FAttributeMap& Attributes) override
 	{
 		check(InputTensors.Num() >= 3);
 		check(InputTensors.Num() <= 5);
diff --git a/Engine/Plugins/Experimental/NNERuntimeRDG/Source/NNERuntimeRDG/Private/Dml/NNEDmlOperatorSplit.cpp b/Engine/Plugins/Experimental/NNERuntimeRDG/Source/NNERuntimeRDG/Private/Dml/NNEDmlOperatorSplit.cpp
--- a/Engine/Plugins/Experimental/NNERuntimeRDG/Source/NNERuntimeRDG/Private/Dml/NNEDmlOperatorSplit.cpp
+++ b/Engine/Plugins/Experimental/NNERuntimeRDG/Source/NNERuntimeRDG/Private/Dml/NNEDmlOperatorSplit.cpp
@@ -17,7 +17,7 @@ public:
 		return new FOperatorDmlSplit();
 	}
 
-	static bool Validate(const NNECore::FAttributeMap& AttributeMap, TConstArrayView<ENNETensorDataType> InputTypes, TConstArrayView<NNECore::FSymbolicTensorShape> InputShapes)
+	static bool Validate(const NNE::FAttributeMap& AttributeMap, TConstArrayView<ENNETensorDataType> InputTypes, TConstArrayView<NNE::FSymbolicTensorShape> InputShapes)
 	{
 		//TODO
 		return true;
@@ -26,10 +26,10 @@ public:
 	//
 	//
 	//
-	virtual bool Initialize(IDMLDevice* Device, TArrayView<const NNECore::Internal::FTensor> InputTensors, TArrayView<const NNECore::Internal::FTensor> OutputTensors, const NNECore::FAttributeMap& Attributes) override
+	virtual bool Initialize(IDMLDevice* Device, TArrayView<const NNE::Internal::FTensor> InputTensors, TArrayView<const NNE::Internal::FTensor> OutputTensors, const NNE::FAttributeMap& Attributes) override
 	{
 		check(InputTensors.Num() == 1);
-		const NNECore::Internal::FTensor& InputTensor = InputTensors[0];
+		const NNE::Internal::FTensor& InputTensor = InputTensors[0];
 
 		int Axis = Attributes.GetValueOrDefault<int>(TEXT("axis"), 0);
 
diff --git a/Engine/Plugins/Experimental/NNERuntimeRDG/Source/NNERuntimeRDG/Private/Dml/NNEDmlOperatorUtils.h b/Engine/Plugins/Experimental/NNERuntimeRDG/Source/NNERuntimeRDG/Private/Dml/NNEDmlOperatorUtils.h
--- a/Engine/Plugins/Experimental/NNERuntimeRDG/Source/NNERuntimeRDG/Private/Dml/NNEDmlOperatorUtils.h
+++ b/Engine/Plugins/Experimental/NNERuntimeRDG/Source/NNERuntimeRDG/Private/Dml/NNEDmlOperatorUtils.h
@@ -65,7 +65,7 @@ static Util::FSmallUIntArray KernelPadding(
 
 static void ComputeStartEndPaddings(
     TConstArrayView<uint32> InputShape,
-    const NNECore::FAttributeMap& Attributes, 
+    const NNE::FAttributeMap& Attributes, 
     Util::FSmallUIntArray &OutStartPadding, 
 	Util::FSmallUIntArray &OutEndPadding,
     TConstArrayView<uint32> Padding)
diff --git a/Engine/Plugins/Experimental/NNERuntimeRDG/Source/NNERuntimeRDG/Private/Helper/NNERuntimeRDGHelperCast.cpp b/Engine/Plugins/Experimental/NNERuntimeRDG/Source/NNERuntimeRDG/Private/Helper/NNERuntimeRDGHelperCast.cpp
--- a/Engine/Plugins/Experimental/NNERuntimeRDG/Source/NNERuntimeRDG/Private/Helper/NNERuntimeRDGHelperCast.cpp
+++ b/Engine/Plugins/Experimental/NNERuntimeRDG/Source/NNERuntimeRDG/Private/Helper/NNERuntimeRDGHelperCast.cpp
@@ -16,7 +16,7 @@ namespace UE::NNERuntimeRDG::Internal::CPUHelper::Cast
 		}
 	}
 
-	template<typename TOutput> void ApplyWithResolvedOutputType(const NNECore::Internal::FTensor& Tensor, NNECore::Internal::FTensor& OutputTensor)
+	template<typename TOutput> void ApplyWithResolvedOutputType(const NNE::Internal::FTensor& Tensor, NNE::Internal::FTensor& OutputTensor)
 	{
 		ENNETensorDataType SourceType = Tensor.GetDataType();
 		TArray<TOutput> OutputData;
@@ -43,10 +43,10 @@ namespace UE::NNERuntimeRDG::Internal::CPUHelper::Cast
 		OutputTensor.SetPreparedData<TOutput>(OutputData);
 	}
 
-	void Apply(const NNECore::Internal::FTensor& Tensor, NNECore::Internal::FTensor& OutputTensor)
+	void Apply(const NNE::Internal::FTensor& Tensor, NNE::Internal::FTensor& OutputTensor)
 	{
 		//Heuristic to avoid unexpected performance hit. This helper being intended for shape related arithmetic only.
-		static constexpr int32 MaxItemInInputTensors = NNECore::FTensorShape::MaxRank * 2;
+		static constexpr int32 MaxItemInInputTensors = NNE::FTensorShape::MaxRank * 2;
 
 		if (Tensor.HasPreparedData() && (Tensor.GetVolume() <= MaxItemInInputTensors))
 		{
diff --git a/Engine/Plugins/Experimental/NNERuntimeRDG/Source/NNERuntimeRDG/Private/Helper/NNERuntimeRDGHelperConcat.cpp b/Engine/Plugins/Experimental/NNERuntimeRDG/Source/NNERuntimeRDG/Private/Helper/NNERuntimeRDGHelperConcat.cpp
--- a/Engine/Plugins/Experimental/NNERuntimeRDG/Source/NNERuntimeRDG/Private/Helper/NNERuntimeRDGHelperConcat.cpp
+++ b/Engine/Plugins/Experimental/NNERuntimeRDG/Source/NNERuntimeRDG/Private/Helper/NNERuntimeRDGHelperConcat.cpp
@@ -7,9 +7,9 @@
 
 namespace UE::NNERuntimeRDG::Internal::CPUHelper::Concat
 {
-	void Apply(TConstArrayView<NNECore::Internal::FTensorRef> InputTensors, NNECore::Internal::FTensor& OutputTensor, int32 Axis)
+	void Apply(TConstArrayView<NNE::Internal::FTensorRef> InputTensors, NNE::Internal::FTensor& OutputTensor, int32 Axis)
 	{
-		static constexpr int32 MaxItemInOutputTensor = NNECore::FTensorShape::MaxRank * 2;
+		static constexpr int32 MaxItemInOutputTensor = NNE::FTensorShape::MaxRank * 2;
 		
 		if (OutputTensor.GetVolume() >= MaxItemInOutputTensor)
 		{
@@ -21,7 +21,7 @@ namespace UE::NNERuntimeRDG::Internal::CPUHelper::Concat
 			return;
 		}
 		
-		for (NNECore::Internal::FTensorRef InputTensor : InputTensors)
+		for (NNE::Internal::FTensorRef InputTensor : InputTensors)
 		{
 			check(InputTensor != nullptr);
 			if (!InputTensor->HasPreparedData())
@@ -38,7 +38,7 @@ namespace UE::NNERuntimeRDG::Internal::CPUHelper::Concat
 
 		OutputData.SetNumUninitialized(OutputTensor.GetVolume());
 		
-		for (NNECore::Internal::FTensorRef InputTensor : InputTensors)
+		for (NNE::Internal::FTensorRef InputTensor : InputTensors)
 		{
 			TConstArrayView<float> InputData = InputTensor->GetPreparedData<float>();
 			Private::TensorIdxIterator itInput(InputTensor->GetShape());
diff --git a/Engine/Plugins/Experimental/NNERuntimeRDG/Source/NNERuntimeRDG/Private/Helper/NNERuntimeRDGHelperElementWiseBinary.cpp b/Engine/Plugins/Experimental/NNERuntimeRDG/Source/NNERuntimeRDG/Private/Helper/NNERuntimeRDGHelperElementWiseBinary.cpp
--- a/Engine/Plugins/Experimental/NNERuntimeRDG/Source/NNERuntimeRDG/Private/Helper/NNERuntimeRDGHelperElementWiseBinary.cpp
+++ b/Engine/Plugins/Experimental/NNERuntimeRDG/Source/NNERuntimeRDG/Private/Helper/NNERuntimeRDGHelperElementWiseBinary.cpp
@@ -9,36 +9,36 @@
 
 namespace UE::NNERuntimeRDG::Internal::CPUHelper::ElementWiseBinary
 {
-	template<NNECore::Internal::EElementWiseBinaryOperatorType OpType, typename TData> TData Apply(TData X, TData Y);
+	template<NNE::Internal::EElementWiseBinaryOperatorType OpType, typename TData> TData Apply(TData X, TData Y);
 
-	template<> float Apply<NNECore::Internal::EElementWiseBinaryOperatorType::Add, float>(float X, float Y) { return X + Y; }
-	template<> float Apply<NNECore::Internal::EElementWiseBinaryOperatorType::Div, float>(float X, float Y) { return X / Y; }
-	template<> float Apply<NNECore::Internal::EElementWiseBinaryOperatorType::Mod, float>(float X, float Y) { return FMath::Fmod(X, Y); }
-	template<> float Apply<NNECore::Internal::EElementWiseBinaryOperatorType::Mul, float>(float X, float Y) { return X * Y; }
-	template<> float Apply<NNECore::Internal::EElementWiseBinaryOperatorType::Prelu, float>(float X, float Y) { return (X < 0.0f) ? (Y * X) : X; }
-	template<> float Apply<NNECore::Internal::EElementWiseBinaryOperatorType::Pow, float>(float X, float Y) { return FMath::Pow(X, Y); }
-	template<> float Apply<NNECore::Internal::EElementWiseBinaryOperatorType::Sub, float>(float X, float Y) { return X - Y; }
+	template<> float Apply<NNE::Internal::EElementWiseBinaryOperatorType::Add, float>(float X, float Y) { return X + Y; }
+	template<> float Apply<NNE::Internal::EElementWiseBinaryOperatorType::Div, float>(float X, float Y) { return X / Y; }
+	template<> float Apply<NNE::Internal::EElementWiseBinaryOperatorType::Mod, float>(float X, float Y) { return FMath::Fmod(X, Y); }
+	template<> float Apply<NNE::Internal::EElementWiseBinaryOperatorType::Mul, float>(float X, float Y) { return X * Y; }
+	template<> float Apply<NNE::Internal::EElementWiseBinaryOperatorType::Prelu, float>(float X, float Y) { return (X < 0.0f) ? (Y * X) : X; }
+	template<> float Apply<NNE::Internal::EElementWiseBinaryOperatorType::Pow, float>(float X, float Y) { return FMath::Pow(X, Y); }
+	template<> float Apply<NNE::Internal::EElementWiseBinaryOperatorType::Sub, float>(float X, float Y) { return X - Y; }
 
-	template<> int32 Apply<NNECore::Internal::EElementWiseBinaryOperatorType::Add, int32>(int32 X, int32 Y) { return X + Y; }
-	template<> int32 Apply<NNECore::Internal::EElementWiseBinaryOperatorType::Div, int32>(int32 X, int32 Y) { return X / Y; }
-	template<> int32 Apply<NNECore::Internal::EElementWiseBinaryOperatorType::Mod, int32>(int32 X, int32 Y) { return X % Y; }
-	template<> int32 Apply<NNECore::Internal::EElementWiseBinaryOperatorType::Mul, int32>(int32 X, int32 Y) { return X * Y; }
-	template<> int32 Apply<NNECore::Internal::EElementWiseBinaryOperatorType::Prelu, int32>(int32 X, int32 Y) { return (X < 0.0f) ? (Y * X) : X; }
-	template<> int32 Apply<NNECore::Internal::EElementWiseBinaryOperatorType::Pow, int32>(int32 X, int32 Y) { return (int32)FMath::Pow((float)X, (float)Y); }
-	template<> int32 Apply<NNECore::Internal::EElementWiseBinaryOperatorType::Sub, int32>(int32 X, int32 Y) { return X - Y; }
+	template<> int32 Apply<NNE::Internal::EElementWiseBinaryOperatorType::Add, int32>(int32 X, int32 Y) { return X + Y; }
+	template<> int32 Apply<NNE::Internal::EElementWiseBinaryOperatorType::Div, int32>(int32 X, int32 Y) { return X / Y; }
+	template<> int32 Apply<NNE::Internal::EElementWiseBinaryOperatorType::Mod, int32>(int32 X, int32 Y) { return X % Y; }
+	template<> int32 Apply<NNE::Internal::EElementWiseBinaryOperatorType::Mul, int32>(int32 X, int32 Y) { return X * Y; }
+	template<> int32 Apply<NNE::Internal::EElementWiseBinaryOperatorType::Prelu, int32>(int32 X, int32 Y) { return (X < 0.0f) ? (Y * X) : X; }
+	template<> int32 Apply<NNE::Internal::EElementWiseBinaryOperatorType::Pow, int32>(int32 X, int32 Y) { return (int32)FMath::Pow((float)X, (float)Y); }
+	template<> int32 Apply<NNE::Internal::EElementWiseBinaryOperatorType::Sub, int32>(int32 X, int32 Y) { return X - Y; }
 
-	template<> int64 Apply<NNECore::Internal::EElementWiseBinaryOperatorType::Add, int64>(int64 X, int64 Y) { return X + Y; }
-	template<> int64 Apply<NNECore::Internal::EElementWiseBinaryOperatorType::Div, int64>(int64 X, int64 Y) { return X / Y; }
-	template<> int64 Apply<NNECore::Internal::EElementWiseBinaryOperatorType::Mod, int64>(int64 X, int64 Y) { return X % Y; }
-	template<> int64 Apply<NNECore::Internal::EElementWiseBinaryOperatorType::Mul, int64>(int64 X, int64 Y) { return X * Y; }
-	template<> int64 Apply<NNECore::Internal::EElementWiseBinaryOperatorType::Prelu, int64>(int64 X, int64 Y) { return (X < 0.0f) ? (Y * X) : X; }
-	template<> int64 Apply<NNECore::Internal::EElementWiseBinaryOperatorType::Pow, int64>(int64 X, int64 Y) { return (int64)FMath::Pow((float)X, (float)Y); }
-	template<> int64 Apply<NNECore::Internal::EElementWiseBinaryOperatorType::Sub, int64>(int64 X, int64 Y) { return X - Y; }
+	template<> int64 Apply<NNE::Internal::EElementWiseBinaryOperatorType::Add, int64>(int64 X, int64 Y) { return X + Y; }
+	template<> int64 Apply<NNE::Internal::EElementWiseBinaryOperatorType::Div, int64>(int64 X, int64 Y) { return X / Y; }
+	template<> int64 Apply<NNE::Internal::EElementWiseBinaryOperatorType::Mod, int64>(int64 X, int64 Y) { return X % Y; }
+	template<> int64 Apply<NNE::Internal::EElementWiseBinaryOperatorType::Mul, int64>(int64 X, int64 Y) { return X * Y; }
+	template<> int64 Apply<NNE::Internal::EElementWiseBinaryOperatorType::Prelu, int64>(int64 X, int64 Y) { return (X < 0.0f) ? (Y * X) : X; }
+	template<> int64 Apply<NNE::Internal::EElementWiseBinaryOperatorType::Pow, int64>(int64 X, int64 Y) { return (int64)FMath::Pow((float)X, (float)Y); }
+	template<> int64 Apply<NNE::Internal::EElementWiseBinaryOperatorType::Sub, int64>(int64 X, int64 Y) { return X - Y; }
 
-	template<NNECore::Internal::EElementWiseBinaryOperatorType OpType, typename TData> void Apply(const NNECore::Internal::FTensor& LHSTensor, const NNECore::Internal::FTensor& RHSTensor, NNECore::Internal::FTensor& OutputTensor)
+	template<NNE::Internal::EElementWiseBinaryOperatorType OpType, typename TData> void Apply(const NNE::Internal::FTensor& LHSTensor, const NNE::Internal::FTensor& RHSTensor, NNE::Internal::FTensor& OutputTensor)
 	{
 		//Heuristic to avoid unexpected performance hit. This helper being intended for shape related arithmetic only.
-		static constexpr int32 MaxItemInInputTensors = NNECore::FTensorShape::MaxRank * 2;
+		static constexpr int32 MaxItemInInputTensors = NNE::FTensorShape::MaxRank * 2;
 
 		if (LHSTensor.HasPreparedData() && 
 			RHSTensor.HasPreparedData() && 
@@ -65,37 +65,37 @@ namespace UE::NNERuntimeRDG::Internal::CPUHelper::ElementWiseBinary
 		}
 	}
 
-	template<typename TData> void ApplyResolvedDataType(NNECore::Internal::EElementWiseBinaryOperatorType OpType, const NNECore::Internal::FTensor& LHSTensor, const NNECore::Internal::FTensor& RHSTensor, NNECore::Internal::FTensor& OutputTensor)
+	template<typename TData> void ApplyResolvedDataType(NNE::Internal::EElementWiseBinaryOperatorType OpType, const NNE::Internal::FTensor& LHSTensor, const NNE::Internal::FTensor& RHSTensor, NNE::Internal::FTensor& OutputTensor)
 	{
 		switch (OpType)
 		{
-		case NNECore::Internal::EElementWiseBinaryOperatorType::Add:
-			Apply<NNECore::Internal::EElementWiseBinaryOperatorType::Add, TData>(LHSTensor, RHSTensor, OutputTensor);
+		case NNE::Internal::EElementWiseBinaryOperatorType::Add:
+			Apply<NNE::Internal::EElementWiseBinaryOperatorType::Add, TData>(LHSTensor, RHSTensor, OutputTensor);
 			break;
-		case NNECore::Internal::EElementWiseBinaryOperatorType::Div:
-			Apply<NNECore::Internal::EElementWiseBinaryOperatorType::Div, TData>(LHSTensor, RHSTensor, OutputTensor);
+		case NNE::Internal::EElementWiseBinaryOperatorType::Div:
+			Apply<NNE::Internal::EElementWiseBinaryOperatorType::Div, TData>(LHSTensor, RHSTensor, OutputTensor);
 			break;
-		case NNECore::Internal::EElementWiseBinaryOperatorType::Mod:
-			Apply<NNECore::Internal::EElementWiseBinaryOperatorType::Mod, TData>(LHSTensor, RHSTensor, OutputTensor);
+		case NNE::Internal::EElementWiseBinaryOperatorType::Mod:
+			Apply<NNE::Internal::EElementWiseBinaryOperatorType::Mod, TData>(LHSTensor, RHSTensor, OutputTensor);
 			break;
-		case NNECore::Internal::EElementWiseBinaryOperatorType::Mul:
-			Apply<NNECore::Internal::EElementWiseBinaryOperatorType::Mul, TData>(LHSTensor, RHSTensor, OutputTensor);
+		case NNE::Internal::EElementWiseBinaryOperatorType::Mul:
+			Apply<NNE::Internal::EElementWiseBinaryOperatorType::Mul, TData>(LHSTensor, RHSTensor, OutputTensor);
 			break;
-		case NNECore::Internal::EElementWiseBinaryOperatorType::Prelu:
-			Apply<NNECore::Internal::EElementWiseBinaryOperatorType::Prelu, TData>(LHSTensor, RHSTensor, OutputTensor);
+		case NNE::Internal::EElementWiseBinaryOperatorType::Prelu:
+			Apply<NNE::Internal::EElementWiseBinaryOperatorType::Prelu, TData>(LHSTensor, RHSTensor, OutputTensor);
 			break;
-		case NNECore::Internal::EElementWiseBinaryOperatorType::Pow:
-			Apply<NNECore::Internal::EElementWiseBinaryOperatorType::Pow, TData>(LHSTensor, RHSTensor, OutputTensor);
+		case NNE::Internal::EElementWiseBinaryOperatorType::Pow:
+			Apply<NNE::Internal::EElementWiseBinaryOperatorType::Pow, TData>(LHSTensor, RHSTensor, OutputTensor);
 			break;
-		case NNECore::Internal::EElementWiseBinaryOperatorType::Sub:
-			Apply<NNECore::Internal::EElementWiseBinaryOperatorType::Sub, TData>(LHSTensor, RHSTensor, OutputTensor);
+		case NNE::Internal::EElementWiseBinaryOperatorType::Sub:
+			Apply<NNE::Internal::EElementWiseBinaryOperatorType::Sub, TData>(LHSTensor, RHSTensor, OutputTensor);
 			break;
 		default:
 			break;
 		}
 	}
 
-	void Apply(NNECore::Internal::EElementWiseBinaryOperatorType OpType, const NNECore::Internal::FTensor& LHSTensor, const NNECore::Internal::FTensor& RHSTensor, NNECore::Internal::FTensor& OutputTensor)
+	void Apply(NNE::Internal::EElementWiseBinaryOperatorType OpType, const NNE::Internal::FTensor& LHSTensor, const NNE::Internal::FTensor& RHSTensor, NNE::Internal::FTensor& OutputTensor)
 	{
 		switch (OutputTensor.GetDataType())
 		{
diff --git a/Engine/Plugins/Experimental/NNERuntimeRDG/Source/NNERuntimeRDG/Private/Helper/NNERuntimeRDGHelperElementWiseUnary.cpp b/Engine/Plugins/Experimental/NNERuntimeRDG/Source/NNERuntimeRDG/Private/Helper/NNERuntimeRDGHelperElementWiseUnary.cpp
--- a/Engine/Plugins/Experimental/NNERuntimeRDG/Source/NNERuntimeRDG/Private/Helper/NNERuntimeRDGHelperElementWiseUnary.cpp
+++ b/Engine/Plugins/Experimental/NNERuntimeRDG/Source/NNERuntimeRDG/Private/Helper/NNERuntimeRDGHelperElementWiseUnary.cpp
@@ -6,13 +6,13 @@
 
 namespace UE::NNERuntimeRDG::Internal::CPUHelper::ElementWiseUnary
 {
-	template<NNECore::Internal::EElementWiseUnaryOperatorType OpType> float Apply(float X, float Alpha, float Beta, float Gamma);
+	template<NNE::Internal::EElementWiseUnaryOperatorType OpType> float Apply(float X, float Alpha, float Beta, float Gamma);
 	
-	template<> float Apply<NNECore::Internal::EElementWiseUnaryOperatorType::Abs>(float X, float Alpha, float Beta, float Gamma) { return FMath::Abs(X); }
+	template<> float Apply<NNE::Internal::EElementWiseUnaryOperatorType::Abs>(float X, float Alpha, float Beta, float Gamma) { return FMath::Abs(X); }
 
-	template<> float Apply<NNECore::Internal::EElementWiseUnaryOperatorType::Acos>(float X, float Alpha, float Beta, float Gamma) { return FMath::Acos(X); }
+	template<> float Apply<NNE::Internal::EElementWiseUnaryOperatorType::Acos>(float X, float Alpha, float Beta, float Gamma) { return FMath::Acos(X); }
 	
-	template<> float Apply<NNECore::Internal::EElementWiseUnaryOperatorType::Acosh>(float X, float Alpha, float Beta, float Gamma) {
+	template<> float Apply<NNE::Internal::EElementWiseUnaryOperatorType::Acosh>(float X, float Alpha, float Beta, float Gamma) {
 		//https://mathworld.wolfram.com/InverseHyperbolicCosine.html
 		float FloatNan = FMath::Sqrt(-1.0f);
 		float yAboveOne = FMath::Loge(X + FMath::Sqrt(X + 1.0f) * FMath::Sqrt(X - 1.0f));
@@ -26,110 +26,110 @@ namespace UE::NNERuntimeRDG::Internal::CPUHelper::ElementWiseUnary
 		}
 	}
 	
-	template<> float Apply<NNECore::Internal::EElementWiseUnaryOperatorType::Asin>(float X, float Alpha, float Beta, float Gamma) { return FMath::Asin(X); }
+	template<> float Apply<NNE::Internal::EElementWiseUnaryOperatorType::Asin>(float X, float Alpha, float Beta, float Gamma) { return FMath::Asin(X); }
 	
-	template<> float Apply<NNECore::Internal::EElementWiseUnaryOperatorType::Asinh>(float X, float Alpha, float Beta, float Gamma) {
+	template<> float Apply<NNE::Internal::EElementWiseUnaryOperatorType::Asinh>(float X, float Alpha, float Beta, float Gamma) {
 		//https://mathworld.wolfram.com/InverseHyperbolicSine.html
 		return FMath::Loge(X + FMath::Sqrt(1 + (X * X)));
 	}
 	
-	template<> float Apply<NNECore::Internal::EElementWiseUnaryOperatorType::Atan>(float X, float Alpha, float Beta, float Gamma) { return FMath::Atan(X); }
+	template<> float Apply<NNE::Internal::EElementWiseUnaryOperatorType::Atan>(float X, float Alpha, float Beta, float Gamma) { return FMath::Atan(X); }
 	
-	template<> float Apply<NNECore::Internal::EElementWiseUnaryOperatorType::Atanh>(float X, float Alpha, float Beta, float Gamma) {
+	template<> float Apply<NNE::Internal::EElementWiseUnaryOperatorType::Atanh>(float X, float Alpha, float Beta, float Gamma) {
 		//https://mathworld.wolfram.com/InverseHyperbolicTangent.html
 		return 0.5f * (FMath::Loge(1 + X) - FMath::Loge(1 - X));
 	}
 
-	template<> float Apply<NNECore::Internal::EElementWiseUnaryOperatorType::Ceil>(float X, float Alpha, float Beta, float Gamma) { return FMath::CeilToFloat(X); }
+	template<> float Apply<NNE::Internal::EElementWiseUnaryOperatorType::Ceil>(float X, float Alpha, float Beta, float Gamma) { return FMath::CeilToFloat(X); }
 
-	template<> float Apply<NNECore::Internal::EElementWiseUnaryOperatorType::Clip>(float X, float Alpha, float Beta, float Gamma) { return FMath::Clamp(X, Alpha, Beta); }
+	template<> float Apply<NNE::Internal::EElementWiseUnaryOperatorType::Clip>(float X, float Alpha, float Beta, float Gamma) { return FMath::Clamp(X, Alpha, Beta); }
 	
-	template<> float Apply<NNECore::Internal::EElementWiseUnaryOperatorType::Cos>(float X, float Alpha, float Beta, float Gamma) { return FMath::Cos(X); }
+	template<> float Apply<NNE::Internal::EElementWiseUnaryOperatorType::Cos>(float X, float Alpha, float Beta, float Gamma) { return FMath::Cos(X); }
 	
-	template<> float Apply<NNECore::Internal::EElementWiseUnaryOperatorType::Cosh>(float X, float Alpha, float Beta, float Gamma) {
+	template<> float Apply<NNE::Internal::EElementWiseUnaryOperatorType::Cosh>(float X, float Alpha, float Beta, float Gamma) {
 		//https://mathworld.wolfram.com/HyperbolicCosine.html
 		return 0.5f * (FMath::Exp(X) + FMath::Exp(-X));
 	}
 	
-	template<> float Apply<NNECore::Internal::EElementWiseUnaryOperatorType::Elu>(float X, float Alpha, float Beta, float Gamma) {
+	template<> float Apply<NNE::Internal::EElementWiseUnaryOperatorType::Elu>(float X, float Alpha, float Beta, float Gamma) {
 		//[REDACTED_PROFILE_URL]/onnx/blob/main/docs/Operators.md#elu
 		float yNeg = Alpha * (FMath::Exp(X) - 1.0f);
 		float yPosOrZero = X;
 		return (X >= 0.0f) ? yPosOrZero : yNeg;
 	}
 	
-	template<> float Apply<NNECore::Internal::EElementWiseUnaryOperatorType::Exp>(float X, float Alpha, float Beta, float Gamma) { return FMath::Exp(X); }
+	template<> float Apply<NNE::Internal::EElementWiseUnaryOperatorType::Exp>(float X, float Alpha, float Beta, float Gamma) { return FMath::Exp(X); }
 
-	template<> float Apply<NNECore::Internal::EElementWiseUnaryOperatorType::Floor>(float X, float Alpha, float Beta, float Gamma) { return FMath::Floor(X); }
+	template<> float Apply<NNE::Internal::EElementWiseUnaryOperatorType::Floor>(float X, float Alpha, float Beta, float Gamma) { return FMath::Floor(X); }
 
-	template<> float Apply<NNECore::Internal::EElementWiseUnaryOperatorType::IsInf>(float X, float Alpha, float Beta, float Gamma) { return !FMath::IsFinite(X); }
+	template<> float Apply<NNE::Internal::EElementWiseUnaryOperatorType::IsInf>(float X, float Alpha, float Beta, float Gamma) { return !FMath::IsFinite(X); }
 
-	template<> float Apply<NNECore::Internal::EElementWiseUnaryOperatorType::IsNan>(float X, float Alpha, float Beta, float Gamma) { return FMath::IsNaN(X); }
+	template<> float Apply<NNE::Internal::EElementWiseUnaryOperatorType::IsNan>(float X, float Alpha, float Beta, float Gamma) { return FMath::IsNaN(X); }
 
-	template<> float Apply<NNECore::Internal::EElementWiseUnaryOperatorType::HardSigmoid>(float X, float Alpha, float Beta, float Gamma) {
+	template<> float Apply<NNE::Internal::EElementWiseUnaryOperatorType::HardSigmoid>(float X, float Alpha, float Beta, float Gamma) {
 		//[REDACTED_PROFILE_URL]/onnx/blob/main/docs/Operators.md#hardSigmoid
 		return FMath::Max(0.0f, FMath::Min(1.0f, Alpha * X + Beta));
 	}
 
-	template<> float Apply<NNECore::Internal::EElementWiseUnaryOperatorType::HardSwish>(float X, float Alpha, float Beta, float Gamma) {
+	template<> float Apply<NNE::Internal::EElementWiseUnaryOperatorType::HardSwish>(float X, float Alpha, float Beta, float Gamma) {
 		//[REDACTED_PROFILE_URL]/onnx/blob/main/docs/Operators.md#hardSwish
-		return Apply<NNECore::Internal::EElementWiseUnaryOperatorType::HardSigmoid>(X, 1.0f / 6.0f, 0.5f, Gamma);
+		return Apply<NNE::Internal::EElementWiseUnaryOperatorType::HardSigmoid>(X, 1.0f / 6.0f, 0.5f, Gamma);
 	}
 
-	template<> float Apply<NNECore::Internal::EElementWiseUnaryOperatorType::LeakyRelu>(float X, float Alpha, float Beta, float Gamma) { return (X >= 0.0f) ? X : Alpha * X; }
+	template<> float Apply<NNE::Internal::EElementWiseUnaryOperatorType::LeakyRelu>(float X, float Alpha, float Beta, float Gamma) { return (X >= 0.0f) ? X : Alpha * X; }
 
-	template<> float Apply<NNECore::Internal::EElementWiseUnaryOperatorType::Log>(float X, float Alpha, float Beta, float Gamma) { return FMath::Loge(X); }
+	template<> float Apply<NNE::Internal::EElementWiseUnaryOperatorType::Log>(float X, float Alpha, float Beta, float Gamma) { return FMath::Loge(X); }
 
-	template<> float Apply<NNECore::Internal::EElementWiseUnaryOperatorType::Neg>(float X, float Alpha, float Beta, float Gamma) { return -X; }
+	template<> float Apply<NNE::Internal::EElementWiseUnaryOperatorType::Neg>(float X, float Alpha, float Beta, float Gamma) { return -X; }
 
-	template<> float Apply<NNECore::Internal::EElementWiseUnaryOperatorType::Reciprocal>(float X, float Alpha, float Beta, float Gamma) { return 1.0f/X; }
+	template<> float Apply<NNE::Internal::EElementWiseUnaryOperatorType::Reciprocal>(float X, float Alpha, float Beta, float Gamma) { return 1.0f/X; }
 
-	template<> float Apply<NNECore::Internal::EElementWiseUnaryOperatorType::Relu>(float X, float Alpha, float Beta, float Gamma) { return FMath::Max(X, 0.0f); }
+	template<> float Apply<NNE::Internal::EElementWiseUnaryOperatorType::Relu>(float X, float Alpha, float Beta, float Gamma) { return FMath::Max(X, 0.0f); }
 
-	template<> float Apply<NNECore::Internal::EElementWiseUnaryOperatorType::Round>(float X, float Alpha, float Beta, float Gamma) { return FMath::RoundToFloat(X); }
+	template<> float Apply<NNE::Internal::EElementWiseUnaryOperatorType::Round>(float X, float Alpha, float Beta, float Gamma) { return FMath::RoundToFloat(X); }
 
-	template<> float Apply<NNECore::Internal::EElementWiseUnaryOperatorType::Selu>(float X, float Alpha, float Beta, float Gamma) {
+	template<> float Apply<NNE::Internal::EElementWiseUnaryOperatorType::Selu>(float X, float Alpha, float Beta, float Gamma) {
 		//[REDACTED_PROFILE_URL]/onnx/blob/main/docs/Operators.md#Selu
 		float yNegOrZero = Gamma * (Alpha * FMath::Exp(X) - Alpha);
 		float yPos = Gamma * X;
 		return (X > 0.0f) ? yPos : yNegOrZero;
 	}
 	
-	template<> float Apply<NNECore::Internal::EElementWiseUnaryOperatorType::Sigmoid>(float X, float Alpha, float Beta, float Gamma) { return 1.0f / (1.0f + FMath::Exp(-X)); }
+	template<> float Apply<NNE::Internal::EElementWiseUnaryOperatorType::Sigmoid>(float X, float Alpha, float Beta, float Gamma) { return 1.0f / (1.0f + FMath::Exp(-X)); }
 	
-	template<> float Apply<NNECore::Internal::EElementWiseUnaryOperatorType::Sign>(float X, float Alpha, float Beta, float Gamma) { return FMath::Sign(X); }
+	template<> float Apply<NNE::Internal::EElementWiseUnaryOperatorType::Sign>(float X, float Alpha, float Beta, float Gamma) { return FMath::Sign(X); }
 	
-	template<> float Apply<NNECore::Internal::EElementWiseUnaryOperatorType::Sin>(float X, float Alpha, float Beta, float Gamma) { return FMath::Sin(X); }
+	template<> float Apply<NNE::Internal::EElementWiseUnaryOperatorType::Sin>(float X, float Alpha, float Beta, float Gamma) { return FMath::Sin(X); }
 	
-	template<> float Apply<NNECore::Internal::EElementWiseUnaryOperatorType::Sinh>(float X, float Alpha, float Beta, float Gamma) { return FMath::Sinh(X); }
+	template<> float Apply<NNE::Internal::EElementWiseUnaryOperatorType::Sinh>(float X, float Alpha, float Beta, float Gamma) { return FMath::Sinh(X); }
 	
-	template<> float Apply<NNECore::Internal::EElementWiseUnaryOperatorType::Softplus>(float X, float Alpha, float Beta, float Gamma) { return FMath::Loge(FMath::Exp(X) + 1.0f); }
+	template<> float Apply<NNE::Internal::EElementWiseUnaryOperatorType::Softplus>(float X, float Alpha, float Beta, float Gamma) { return FMath::Loge(FMath::Exp(X) + 1.0f); }
 	
-	template<> float Apply<NNECore::Internal::EElementWiseUnaryOperatorType::Softsign>(float X, float Alpha, float Beta, float Gamma) { return X / (1.0f + FMath::Abs(X)); }
+	template<> float Apply<NNE::Internal::EElementWiseUnaryOperatorType::Softsign>(float X, float Alpha, float Beta, float Gamma) { return X / (1.0f + FMath::Abs(X)); }
 	
-	template<> float Apply<NNECore::Internal::EElementWiseUnaryOperatorType::Sqrt>(float X, float Alpha, float Beta, float Gamma) { return FMath::Sqrt(X); }
+	template<> float Apply<NNE::Internal::EElementWiseUnaryOperatorType::Sqrt>(float X, float Alpha, float Beta, float Gamma) { return FMath::Sqrt(X); }
 	
-	template<> float Apply<NNECore::Internal::EElementWiseUnaryOperatorType::Tan>(float X, float Alpha, float Beta, float Gamma) { return FMath::Tan(X); }
+	template<> float Apply<NNE::Internal::EElementWiseUnaryOperatorType::Tan>(float X, float Alpha, float Beta, float Gamma) { return FMath::Tan(X); }
 
-	template<> float Apply<NNECore::Internal::EElementWiseUnaryOperatorType::Tanh>(float X, float Alpha, float Beta, float Gamma) {
+	template<> float Apply<NNE::Internal::EElementWiseUnaryOperatorType::Tanh>(float X, float Alpha, float Beta, float Gamma) {
 		//https://mathworld.wolfram.com/HyperbolicTangent.html
-		float SinhValue = Apply<NNECore::Internal::EElementWiseUnaryOperatorType::Sinh>(X, Alpha, Beta, Gamma);
-		float CoshValue = Apply<NNECore::Internal::EElementWiseUnaryOperatorType::Cosh>(X, Alpha, Beta, Gamma);
+		float SinhValue = Apply<NNE::Internal::EElementWiseUnaryOperatorType::Sinh>(X, Alpha, Beta, Gamma);
+		float CoshValue = Apply<NNE::Internal::EElementWiseUnaryOperatorType::Cosh>(X, Alpha, Beta, Gamma);
 		return SinhValue / CoshValue;
 	}
 
-	template<> float Apply<NNECore::Internal::EElementWiseUnaryOperatorType::Erf>(float X, float Alpha, float Beta, float Gamma) {
+	template<> float Apply<NNE::Internal::EElementWiseUnaryOperatorType::Erf>(float X, float Alpha, float Beta, float Gamma) {
 		//https://aapt.scitation.org/doi/abs/10.1119/1.15018?journalCode=ajp
 		float a = 167.0f / 148.0f;
 		float b = 11.0f / 109.0f;
 		float x3 = X * X * X;
-		return Apply<NNECore::Internal::EElementWiseUnaryOperatorType::Tanh>(a * X + b * x3, Alpha, Beta, Gamma);
+		return Apply<NNE::Internal::EElementWiseUnaryOperatorType::Tanh>(a * X + b * x3, Alpha, Beta, Gamma);
 	}
 
-	template<NNECore::Internal::EElementWiseUnaryOperatorType OpType> void Apply(const NNECore::Internal::FTensor& Tensor, float Alpha, float Beta, float Gamma, NNECore::Internal::FTensor& OutputTensor)
+	template<NNE::Internal::EElementWiseUnaryOperatorType OpType> void Apply(const NNE::Internal::FTensor& Tensor, float Alpha, float Beta, float Gamma, NNE::Internal::FTensor& OutputTensor)
 	{
 		//Heuristic to avoid unexpected performance hit. This helper being intended for shape related arithmetic only.
-		static constexpr int32 MaxItemInInputTensors = NNECore::FTensorShape::MaxRank * 2;
+		static constexpr int32 MaxItemInInputTensors = NNE::FTensorShape::MaxRank * 2;
 
 		if (Tensor.HasPreparedData() && (Tensor.GetVolume() <= MaxItemInInputTensors))
 		{
@@ -144,114 +144,114 @@ namespace UE::NNERuntimeRDG::Internal::CPUHelper::ElementWiseUnary
 		}
 	}
 
-	void Apply(NNECore::Internal::EElementWiseUnaryOperatorType OpType, const NNECore::Internal::FTensor& Tensor, float Alpha, float Beta, float Gamma, NNECore::Internal::FTensor& OutputTensor)
+	void Apply(NNE::Internal::EElementWiseUnaryOperatorType OpType, const NNE::Internal::FTensor& Tensor, float Alpha, float Beta, float Gamma, NNE::Internal::FTensor& OutputTensor)
 	{
 		switch (OpType)
 		{
-		case NNECore::Internal::EElementWiseUnaryOperatorType::Abs:
-			Apply<NNECore::Internal::EElementWiseUnaryOperatorType::Abs>(Tensor, Alpha, Beta, Gamma, OutputTensor);
+		case NNE::Internal::EElementWiseUnaryOperatorType::Abs:
+			Apply<NNE::Internal::EElementWiseUnaryOperatorType::Abs>(Tensor, Alpha, Beta, Gamma, OutputTensor);
 			break;
-		case NNECore::Internal::EElementWiseUnaryOperatorType::Acos:
-			Apply<NNECore::Internal::EElementWiseUnaryOperatorType::Acos>(Tensor, Alpha, Beta, Gamma, OutputTensor);
+		case NNE::Internal::EElementWiseUnaryOperatorType::Acos:
+			Apply<NNE::Internal::EElementWiseUnaryOperatorType::Acos>(Tensor, Alpha, Beta, Gamma, OutputTensor);
 			break;
-		case NNECore::Internal::EElementWiseUnaryOperatorType::Acosh:
-			Apply<NNECore::Internal::EElementWiseUnaryOperatorType::Acosh>(Tensor, Alpha, Beta, Gamma, OutputTensor);
+		case NNE::Internal::EElementWiseUnaryOperatorType::Acosh:
+			Apply<NNE::Internal::EElementWiseUnaryOperatorType::Acosh>(Tensor, Alpha, Beta, Gamma, OutputTensor);
 			break;
-		case NNECore::Internal::EElementWiseUnaryOperatorType::Asin:
-			Apply<NNECore::Internal::EElementWiseUnaryOperatorType::Asin>(Tensor, Alpha, Beta, Gamma, OutputTensor);
+		case NNE::Internal::EElementWiseUnaryOperatorType::Asin:
+			Apply<NNE::Internal::EElementWiseUnaryOperatorType::Asin>(Tensor, Alpha, Beta, Gamma, OutputTensor);
 			break;
-		case NNECore::Internal::EElementWiseUnaryOperatorType::Asinh:
-			Apply<NNECore::Internal::EElementWiseUnaryOperatorType::Asinh>(Tensor, Alpha, Beta, Gamma, OutputTensor);
+		case NNE::Internal::EElementWiseUnaryOperatorType::Asinh:
+			Apply<NNE::Internal::EElementWiseUnaryOperatorType::Asinh>(Tensor, Alpha, Beta, Gamma, OutputTensor);
 			break;
-		case NNECore::Internal::EElementWiseUnaryOperatorType::Atan:
-			Apply<NNECore::Internal::EElementWiseUnaryOperatorType::Atan>(Tensor, Alpha, Beta, Gamma, OutputTensor);
+		case NNE::Internal::EElementWiseUnaryOperatorType::Atan:
+			Apply<NNE::Internal::EElementWiseUnaryOperatorType::Atan>(Tensor, Alpha, Beta, Gamma, OutputTensor);
 			break;
-		case NNECore::Internal::EElementWiseUnaryOperatorType::Atanh:
-			Apply<NNECore::Internal::EElementWiseUnaryOperatorType::Atanh>(Tensor, Alpha, Beta, Gamma, OutputTensor);
+		case NNE::Internal::EElementWiseUnaryOperatorType::Atanh:
+			Apply<NNE::Internal::EElementWiseUnaryOperatorType::Atanh>(Tensor, Alpha, Beta, Gamma, OutputTensor);
 			break;
-		case NNECore::Internal::EElementWiseUnaryOperatorType::Ceil:
-			Apply<NNECore::Internal::EElementWiseUnaryOperatorType::Ceil>(Tensor, Alpha, Beta, Gamma, OutputTensor);
+		case NNE::Internal::EElementWiseUnaryOperatorType::Ceil:
+			Apply<NNE::Internal::EElementWiseUnaryOperatorType::Ceil>(Tensor, Alpha, Beta, Gamma, OutputTensor);
 			break;
-		case NNECore::Internal::EElementWiseUnaryOperatorType::Clip:
-			Apply<NNECore::Internal::EElementWiseUnaryOperatorType::Clip>(Tensor, Alpha, Beta, Gamma, OutputTensor);
+		case NNE::Internal::EElementWiseUnaryOperatorType::Clip:
+			Apply<NNE::Internal::EElementWiseUnaryOperatorType::Clip>(Tensor, Alpha, Beta, Gamma, OutputTensor);
 			break;
-		case NNECore::Internal::EElementWiseUnaryOperatorType::Cos:
-			Apply<NNECore::Internal::EElementWiseUnaryOperatorType::Cos>(Tensor, Alpha, Beta, Gamma, OutputTensor);
+		case NNE::Internal::EElementWiseUnaryOperatorType::Cos:
+			Apply<NNE::Internal::EElementWiseUnaryOperatorType::Cos>(Tensor, Alpha, Beta, Gamma, OutputTensor);
 			break;
-		case NNECore::Internal::EElementWiseUnaryOperatorType::Cosh:
-			Apply<NNECore::Internal::EElementWiseUnaryOperatorType::Cosh>(Tensor, Alpha, Beta, Gamma, OutputTensor);
+		case NNE::Internal::EElementWiseUnaryOperatorType::Cosh:
+			Apply<NNE::Internal::EElementWiseUnaryOperatorType::Cosh>(Tensor, Alpha, Beta, Gamma, OutputTensor);
 			break;
-		case NNECore::Internal::EElementWiseUnaryOperatorType::Elu:
-			Apply<NNECore::Internal::EElementWiseUnaryOperatorType::Elu>(Tensor, Alpha, Beta, Gamma, OutputTensor);
+		case NNE::Internal::EElementWiseUnaryOperatorType::Elu:
+			Apply<NNE::Internal::EElementWiseUnaryOperatorType::Elu>(Tensor, Alpha, Beta, Gamma, OutputTensor);
 			break;
-		case NNECore::Internal::EElementWiseUnaryOperatorType::Erf:
-			Apply<NNECore::Internal::EElementWiseUnaryOperatorType::Erf>(Tensor, Alpha, Beta, Gamma, OutputTensor);
+		case NNE::Internal::EElementWiseUnaryOperatorType::Erf:
+			Apply<NNE::Internal::EElementWiseUnaryOperatorType::Erf>(Tensor, Alpha, Beta, Gamma, OutputTensor);
 			break;
-		case NNECore::Internal::EElementWiseUnaryOperatorType::Exp:
-			Apply<NNECore::Internal::EElementWiseUnaryOperatorType::Exp>(Tensor, Alpha, Beta, Gamma, OutputTensor);
+		case NNE::Internal::EElementWiseUnaryOperatorType::Exp:
+			Apply<NNE::Internal::EElementWiseUnaryOperatorType::Exp>(Tensor, Alpha, Beta, Gamma, OutputTensor);
 			break;
-		case NNECore::Internal::EElementWiseUnaryOperatorType::Floor:
-			Apply<NNECore::Internal::EElementWiseUnaryOperatorType::Floor>(Tensor, Alpha, Beta, Gamma, OutputTensor);
+		case NNE::Internal::EElementWiseUnaryOperatorType::Floor:
+			Apply<NNE::Internal::EElementWiseUnaryOperatorType::Floor>(Tensor, Alpha, Beta, Gamma, OutputTensor);
 			break;
-		case NNECore::Internal::EElementWiseUnaryOperatorType::IsInf:
-			Apply<NNECore::Internal::EElementWiseUnaryOperatorType::IsInf>(Tensor, Alpha, Beta, Gamma, OutputTensor);
+		case NNE::Internal::EElementWiseUnaryOperatorType::IsInf:
+			Apply<NNE::Internal::EElementWiseUnaryOperatorType::IsInf>(Tensor, Alpha, Beta, Gamma, OutputTensor);
 			break;
-		case NNECore::Internal::EElementWiseUnaryOperatorType::IsNan:
-			Apply<NNECore::Internal::EElementWiseUnaryOperatorType::IsNan>(Tensor, Alpha, Beta, Gamma, OutputTensor);
+		case NNE::Internal::EElementWiseUnaryOperatorType::IsNan:
+			Apply<NNE::Internal::EElementWiseUnaryOperatorType::IsNan>(Tensor, Alpha, Beta, Gamma, OutputTensor);
 			break;
-		case NNECore::Internal::EElementWiseUnaryOperatorType::HardSigmoid:
-			Apply<NNECore::Internal::EElementWiseUnaryOperatorType::HardSigmoid>(Tensor, Alpha, Beta, Gamma, OutputTensor);
+		case NNE::Internal::EElementWiseUnaryOperatorType::HardSigmoid:
+			Apply<NNE::Internal::EElementWiseUnaryOperatorType::HardSigmoid>(Tensor, Alpha, Beta, Gamma, OutputTensor);
 			break;
-		case NNECore::Internal::EElementWiseUnaryOperatorType::HardSwish:
-			Apply<NNECore::Internal::EElementWiseUnaryOperatorType::HardSwish>(Tensor, Alpha, Beta, Gamma, OutputTensor);
+		case NNE::Internal::EElementWiseUnaryOperatorType::HardSwish:
+			Apply<NNE::Internal::EElementWiseUnaryOperatorType::HardSwish>(Tensor, Alpha, Beta, Gamma, OutputTensor);
 			break;
-		case NNECore::Internal::EElementWiseUnaryOperatorType::LeakyRelu:
-			Apply<NNECore::Internal::EElementWiseUnaryOperatorType::LeakyRelu>(Tensor, Alpha, Beta, Gamma, OutputTensor);
+		case NNE::Internal::EElementWiseUnaryOperatorType::LeakyRelu:
+			Apply<NNE::Internal::EElementWiseUnaryOperatorType::LeakyRelu>(Tensor, Alpha, Beta, Gamma, OutputTensor);
 			break;
-		case NNECore::Internal::EElementWiseUnaryOperatorType::Log:
-			Apply<NNECore::Internal::EElementWiseUnaryOperatorType::Log>(Tensor, Alpha, Beta, Gamma, OutputTensor);
+		case NNE::Internal::EElementWiseUnaryOperatorType::Log:
+			Apply<NNE::Internal::EElementWiseUnaryOperatorType::Log>(Tensor, Alpha, Beta, Gamma, OutputTensor);
 			break;
-		case NNECore::Internal::EElementWiseUnaryOperatorType::Neg:
-			Apply<NNECore::Internal::EElementWiseUnaryOperatorType::Neg>(Tensor, Alpha, Beta, Gamma, OutputTensor);
+		case NNE::Internal::EElementWiseUnaryOperatorType::Neg:
+			Apply<NNE::Internal::EElementWiseUnaryOperatorType::Neg>(Tensor, Alpha, Beta, Gamma, OutputTensor);
 			break;
-		case NNECore::Internal::EElementWiseUnaryOperatorType::Reciprocal:
-			Apply<NNECore::Internal::EElementWiseUnaryOperatorType::Reciprocal>(Tensor, Alpha, Beta, Gamma, OutputTensor);
+		case NNE::Internal::EElementWiseUnaryOperatorType::Reciprocal:
+			Apply<NNE::Internal::EElementWiseUnaryOperatorType::Reciprocal>(Tensor, Alpha, Beta, Gamma, OutputTensor);
 			break;
-		case NNECore::Internal::EElementWiseUnaryOperatorType::Relu:
-			Apply<NNECore::Internal::EElementWiseUnaryOperatorType::Relu>(Tensor, Alpha, Beta, Gamma, OutputTensor);
+		case NNE::Internal::EElementWiseUnaryOperatorType::Relu:
+			Apply<NNE::Internal::EElementWiseUnaryOperatorType::Relu>(Tensor, Alpha, Beta, Gamma, OutputTensor);
 			break;
-		case NNECore::Internal::EElementWiseUnaryOperatorType::Round:
-			Apply<NNECore::Internal::EElementWiseUnaryOperatorType::Round>(Tensor, Alpha, Beta, Gamma, OutputTensor);
+		case NNE::Internal::EElementWiseUnaryOperatorType::Round:
+			Apply<NNE::Internal::EElementWiseUnaryOperatorType::Round>(Tensor, Alpha, Beta, Gamma, OutputTensor);
 			break;
-		case NNECore::Internal::EElementWiseUnaryOperatorType::Selu:
-			Apply<NNECore::Internal::EElementWiseUnaryOperatorType::Selu>(Tensor, Alpha, Beta, Gamma, OutputTensor);
+		case NNE::Internal::EElementWiseUnaryOperatorType::Selu:
+			Apply<NNE::Internal::EElementWiseUnaryOperatorType::Selu>(Tensor, Alpha, Beta, Gamma, OutputTensor);
 			break;
-		case NNECore::Internal::EElementWiseUnaryOperatorType::Sigmoid:
-			Apply<NNECore::Internal::EElementWiseUnaryOperatorType::Sigmoid>(Tensor, Alpha, Beta, Gamma, OutputTensor);
+		case NNE::Internal::EElementWiseUnaryOperatorType::Sigmoid:
+			Apply<NNE::Internal::EElementWiseUnaryOperatorType::Sigmoid>(Tensor, Alpha, Beta, Gamma, OutputTensor);
 			break;
-		case NNECore::Internal::EElementWiseUnaryOperatorType::Sign:
-			Apply<NNECore::Internal::EElementWiseUnaryOperatorType::Sign>(Tensor, Alpha, Beta, Gamma, OutputTensor);
+		case NNE::Internal::EElementWiseUnaryOperatorType::Sign:
+			Apply<NNE::Internal::EElementWiseUnaryOperatorType::Sign>(Tensor, Alpha, Beta, Gamma, OutputTensor);
 			break;
-		case NNECore::Internal::EElementWiseUnaryOperatorType::Sin:
-			Apply<NNECore::Internal::EElementWiseUnaryOperatorType::Sin>(Tensor, Alpha, Beta, Gamma, OutputTensor);
+		case NNE::Internal::EElementWiseUnaryOperatorType::Sin:
+			Apply<NNE::Internal::EElementWiseUnaryOperatorType::Sin>(Tensor, Alpha, Beta, Gamma, OutputTensor);
 			break;
-		case NNECore::Internal::EElementWiseUnaryOperatorType::Sinh:
-			Apply<NNECore::Internal::EElementWiseUnaryOperatorType::Sinh>(Tensor, Alpha, Beta, Gamma, OutputTensor);
+		case NNE::Internal::EElementWiseUnaryOperatorType::Sinh:
+			Apply<NNE::Internal::EElementWiseUnaryOperatorType::Sinh>(Tensor, Alpha, Beta, Gamma, OutputTensor);
 			break;
-		case NNECore::Internal::EElementWiseUnaryOperatorType::Softplus:
-			Apply<NNECore::Internal::EElementWiseUnaryOperatorType::Softplus>(Tensor, Alpha, Beta, Gamma, OutputTensor);
+		case NNE::Internal::EElementWiseUnaryOperatorType::Softplus:
+			Apply<NNE::Internal::EElementWiseUnaryOperatorType::Softplus>(Tensor, Alpha, Beta, Gamma, OutputTensor);
 			break;
-		case NNECore::Internal::EElementWiseUnaryOperatorType::Softsign:
-			Apply<NNECore::Internal::EElementWiseUnaryOperatorType::Softsign>(Tensor, Alpha, Beta, Gamma, OutputTensor);
+		case NNE::Internal::EElementWiseUnaryOperatorType::Softsign:
+			Apply<NNE::Internal::EElementWiseUnaryOperatorType::Softsign>(Tensor, Alpha, Beta, Gamma, OutputTensor);
 			break;
-		case NNECore::Internal::EElementWiseUnaryOperatorType::Sqrt:
-			Apply<NNECore::Internal::EElementWiseUnaryOperatorType::Sqrt>(Tensor, Alpha, Beta, Gamma, OutputTensor);
+		case NNE::Internal::EElementWiseUnaryOperatorType::Sqrt:
+			Apply<NNE::Internal::EElementWiseUnaryOperatorType::Sqrt>(Tensor, Alpha, Beta, Gamma, OutputTensor);
 			break;
-		case NNECore::Internal::EElementWiseUnaryOperatorType::Tan:
-			Apply<NNECore::Internal::EElementWiseUnaryOperatorType::Tan>(Tensor, Alpha, Beta, Gamma, OutputTensor);
+		case NNE::Internal::EElementWiseUnaryOperatorType::Tan:
+			Apply<NNE::Internal::EElementWiseUnaryOperatorType::Tan>(Tensor, Alpha, Beta, Gamma, OutputTensor);
 			break;
-		case NNECore::Internal::EElementWiseUnaryOperatorType::Tanh:
-			Apply<NNECore::Internal::EElementWiseUnaryOperatorType::Tanh>(Tensor, Alpha, Beta, Gamma, OutputTensor);
+		case NNE::Internal::EElementWiseUnaryOperatorType::Tanh:
+			Apply<NNE::Internal::EElementWiseUnaryOperatorType::Tanh>(Tensor, Alpha, Beta, Gamma, OutputTensor);
 			break;
 		default:
 			break;
diff --git a/Engine/Plugins/Experimental/NNERuntimeRDG/Source/NNERuntimeRDG/Private/Helper/NNERuntimeRDGHelperGather.cpp b/Engine/Plugins/Experimental/NNERuntimeRDG/Source/NNERuntimeRDG/Private/Helper/NNERuntimeRDGHelperGather.cpp
--- a/Engine/Plugins/Experimental/NNERuntimeRDG/Source/NNERuntimeRDG/Private/Helper/NNERuntimeRDGHelperGather.cpp
+++ b/Engine/Plugins/Experimental/NNERuntimeRDG/Source/NNERuntimeRDG/Private/Helper/NNERuntimeRDGHelperGather.cpp
@@ -5,7 +5,7 @@
 
 namespace UE::NNERuntimeRDG::Internal::CPUHelper::Gather
 {
-	template <typename TData, typename TInd> void Gather1DResolvedTypes(const NNECore::Internal::FTensor& DataTensor, const NNECore::Internal::FTensor& IndicesTensor, NNECore::Internal::FTensor& OutputTensor)
+	template <typename TData, typename TInd> void Gather1DResolvedTypes(const NNE::Internal::FTensor& DataTensor, const NNE::Internal::FTensor& IndicesTensor, NNE::Internal::FTensor& OutputTensor)
 	{
 		TConstArrayView<TData> DataData = DataTensor.GetPreparedData<TData>();
 		TConstArrayView<TInd> IndicesData = IndicesTensor.GetPreparedData<TInd>();
@@ -22,7 +22,7 @@ namespace UE::NNERuntimeRDG::Internal::CPUHelper::Gather
 		OutputTensor.SetPreparedData<TData>(OutputData);
 	}
 
-	template <typename TInd> void Gather1DResolvedIndices(const NNECore::Internal::FTensor& DataTensor, const NNECore::Internal::FTensor& IndicesTensor, NNECore::Internal::FTensor& OutputTensor)
+	template <typename TInd> void Gather1DResolvedIndices(const NNE::Internal::FTensor& DataTensor, const NNE::Internal::FTensor& IndicesTensor, NNE::Internal::FTensor& OutputTensor)
 	{
 		switch (DataTensor.GetDataType())
 		{
@@ -38,9 +38,9 @@ namespace UE::NNERuntimeRDG::Internal::CPUHelper::Gather
 		}
 	}
 
-	void Apply(const NNECore::Internal::FTensor& DataTensor, const NNECore::Internal::FTensor& IndicesTensor, int32 Axis, NNECore::Internal::FTensor& OutputTensor)
+	void Apply(const NNE::Internal::FTensor& DataTensor, const NNE::Internal::FTensor& IndicesTensor, int32 Axis, NNE::Internal::FTensor& OutputTensor)
 	{
-		static constexpr int32 MaxItemInOutputTensor = NNECore::FTensorShape::MaxRank * 2;
+		static constexpr int32 MaxItemInOutputTensor = NNE::FTensorShape::MaxRank * 2;
 
 		if (OutputTensor.GetVolume() >= MaxItemInOutputTensor)
 		{
diff --git a/Engine/Plugins/Experimental/NNERuntimeRDG/Source/NNERuntimeRDG/Private/Helper/NNERuntimeRDGHelperReshape.h b/Engine/Plugins/Experimental/NNERuntimeRDG/Source/NNERuntimeRDG/Private/Helper/NNERuntimeRDGHelperReshape.h
--- a/Engine/Plugins/Experimental/NNERuntimeRDG/Source/NNERuntimeRDG/Private/Helper/NNERuntimeRDGHelperReshape.h
+++ b/Engine/Plugins/Experimental/NNERuntimeRDG/Source/NNERuntimeRDG/Private/Helper/NNERuntimeRDGHelperReshape.h
@@ -13,13 +13,13 @@ namespace UE::NNERuntimeRDG::Private::ShapeHelper::Reshape
 {
 	template<typename DataType, class OutputShapeType >
 	bool ReshapeTensor(
-		const NNECore::FTensorShape& InputTensorShape,
-		const NNECore::Internal::FTensor& ShapeTensor,
+		const NNE::FTensorShape& InputTensorShape,
+		const NNE::Internal::FTensor& ShapeTensor,
 		bool bAllowZero,
 		OutputShapeType& OutShape)
 	{
 
-		TArray<DataType, TInlineAllocator<NNECore::FTensorShape::MaxRank>> ReshapedShape(ShapeTensor.GetPreparedData<DataType>());
+		TArray<DataType, TInlineAllocator<NNE::FTensorShape::MaxRank>> ReshapedShape(ShapeTensor.GetPreparedData<DataType>());
 
 		if (!bAllowZero)
 		{
diff --git a/Engine/Plugins/Experimental/NNERuntimeRDG/Source/NNERuntimeRDG/Private/Helper/NNERuntimeRDGHelperSlice.cpp b/Engine/Plugins/Experimental/NNERuntimeRDG/Source/NNERuntimeRDG/Private/Helper/NNERuntimeRDGHelperSlice.cpp
--- a/Engine/Plugins/Experimental/NNERuntimeRDG/Source/NNERuntimeRDG/Private/Helper/NNERuntimeRDGHelperSlice.cpp
+++ b/Engine/Plugins/Experimental/NNERuntimeRDG/Source/NNERuntimeRDG/Private/Helper/NNERuntimeRDGHelperSlice.cpp
@@ -5,7 +5,7 @@
 
 namespace UE::NNERuntimeRDG::Internal::CPUHelper::Slice
 {
-	template<typename TData> void ApplyResolvedInputType(const NNECore::Internal::FTensor& InputTensor, NNECore::Internal::FTensor& OutputTensor, TConstArrayView<int32> Starts)
+	template<typename TData> void ApplyResolvedInputType(const NNE::Internal::FTensor& InputTensor, NNE::Internal::FTensor& OutputTensor, TConstArrayView<int32> Starts)
 	{
 		check(InputTensor.HasPreparedData());
 		check(InputTensor.GetShape().Rank() == Starts.Num());
@@ -33,9 +33,9 @@ namespace UE::NNERuntimeRDG::Internal::CPUHelper::Slice
 		OutputTensor.SetPreparedData<TData>(OutputData);
 	}
 
-	void Apply(const NNECore::Internal::FTensor& InputTensor, NNECore::Internal::FTensor& OutputTensor, TConstArrayView<int32> Starts)
+	void Apply(const NNE::Internal::FTensor& InputTensor, NNE::Internal::FTensor& OutputTensor, TConstArrayView<int32> Starts)
 	{
-		static constexpr int32 MaxItemInOutputTensor = NNECore::FTensorShape::MaxRank * 2;
+		static constexpr int32 MaxItemInOutputTensor = NNE::FTensorShape::MaxRank * 2;
 
 		if (OutputTensor.GetVolume() >= MaxItemInOutputTensor)
 		{
diff --git a/Engine/Plugins/Experimental/NNERuntimeRDG/Source/NNERuntimeRDG/Private/Helper/NNERuntimeRDGHelperTranspose.cpp b/Engine/Plugins/Experimental/NNERuntimeRDG/Source/NNERuntimeRDG/Private/Helper/NNERuntimeRDGHelperTranspose.cpp
--- a/Engine/Plugins/Experimental/NNERuntimeRDG/Source/NNERuntimeRDG/Private/Helper/NNERuntimeRDGHelperTranspose.cpp
+++ b/Engine/Plugins/Experimental/NNERuntimeRDG/Source/NNERuntimeRDG/Private/Helper/NNERuntimeRDGHelperTranspose.cpp
@@ -5,7 +5,7 @@
 
 namespace UE::NNERuntimeRDG::Internal::CPUHelper::Transpose
 {
-	bool TransposePreparedData(NNECore::Internal::FTensor& Tensor, TConstArrayView<int32> Perms)
+	bool TransposePreparedData(NNE::Internal::FTensor& Tensor, TConstArrayView<int32> Perms)
 	{
 		if (!Tensor.HasPreparedData())
 		{
@@ -34,7 +34,7 @@ namespace UE::NNERuntimeRDG::Internal::CPUHelper::Transpose
 		TransposedPosition.SetNumUninitialized(TransposedShape.Num());
 
 		Private::TensorIdxIterator it(Tensor.GetShape());
-		const Private::TensorIdxIterator itTransposed(NNECore::FTensorShape::Make(TransposedShape));
+		const Private::TensorIdxIterator itTransposed(NNE::FTensorShape::Make(TransposedShape));
 
 		do
 		{
