NNE Rename NNECore module to NNE following split of runtimes into there own plugins (step2)

diff --git a/Engine/Plugins/Experimental/NNERuntimeRDG/Source/NNERuntimeRDG/Private/Helper/NNERuntimeRDGTensorIdxIterator.cpp b/Engine/Plugins/Experimental/NNERuntimeRDG/Source/NNERuntimeRDG/Private/Helper/NNERuntimeRDGTensorIdxIterator.cpp
--- a/Engine/Plugins/Experimental/NNERuntimeRDG/Source/NNERuntimeRDG/Private/Helper/NNERuntimeRDGTensorIdxIterator.cpp
+++ b/Engine/Plugins/Experimental/NNERuntimeRDG/Source/NNERuntimeRDG/Private/Helper/NNERuntimeRDGTensorIdxIterator.cpp
@@ -4,7 +4,7 @@
 
 namespace UE::NNERuntimeRDG::Private
 {
-	TensorIdxIterator::TensorIdxIterator(const NNECore::FTensorShape& InTensorShape) : TensorShape(InTensorShape)
+	TensorIdxIterator::TensorIdxIterator(const NNE::FTensorShape& InTensorShape) : TensorShape(InTensorShape)
 	{
 		CurrentPosition.Init(0, InTensorShape.Rank());
 	}
@@ -28,7 +28,7 @@ namespace UE::NNERuntimeRDG::Private
 		return CurrentPosition;
 	}
 
-	int32 TensorIdxIterator::GetIndexToBroadcastedShape(const NNECore::FTensorShape& InTensorShape) const
+	int32 TensorIdxIterator::GetIndexToBroadcastedShape(const NNE::FTensorShape& InTensorShape) const
 	{
 		int32 Index = 0;
 		int32 DimBaseOffset = 1;
diff --git a/Engine/Plugins/Experimental/NNERuntimeRDG/Source/NNERuntimeRDG/Private/Helper/NNERuntimeRDGTensorIdxIterator.h b/Engine/Plugins/Experimental/NNERuntimeRDG/Source/NNERuntimeRDG/Private/Helper/NNERuntimeRDGTensorIdxIterator.h
--- a/Engine/Plugins/Experimental/NNERuntimeRDG/Source/NNERuntimeRDG/Private/Helper/NNERuntimeRDGTensorIdxIterator.h
+++ b/Engine/Plugins/Experimental/NNERuntimeRDG/Source/NNERuntimeRDG/Private/Helper/NNERuntimeRDGTensorIdxIterator.h
@@ -9,13 +9,13 @@ namespace UE::NNERuntimeRDG::Private
 {
 	class TensorIdxIterator
 	{
-		const NNECore::FTensorShape& TensorShape;
-		TArray<uint32, TInlineAllocator<NNECore::FTensorShape::MaxRank>> CurrentPosition;
+		const NNE::FTensorShape& TensorShape;
+		TArray<uint32, TInlineAllocator<NNE::FTensorShape::MaxRank>> CurrentPosition;
 
 	public:
-		TensorIdxIterator(const NNECore::FTensorShape& InTensorShape);
+		TensorIdxIterator(const NNE::FTensorShape& InTensorShape);
 		bool Advance();
-		int32 GetIndexToBroadcastedShape(const NNECore::FTensorShape& InTensorShape) const;
+		int32 GetIndexToBroadcastedShape(const NNE::FTensorShape& InTensorShape) const;
 		int32 GetIndex() const;
 		TConstArrayView<uint32> GetPositions() const;
 		int32 GetIndexFromPosition(TConstArrayView<uint32> Position) const;
diff --git a/Engine/Plugins/Experimental/NNERuntimeRDG/Source/NNERuntimeRDG/Private/Hlsl/NNERuntimeRDGBatchNormalization.cpp b/Engine/Plugins/Experimental/NNERuntimeRDG/Source/NNERuntimeRDG/Private/Hlsl/NNERuntimeRDGBatchNormalization.cpp
--- a/Engine/Plugins/Experimental/NNERuntimeRDG/Source/NNERuntimeRDG/Private/Hlsl/NNERuntimeRDGBatchNormalization.cpp
+++ b/Engine/Plugins/Experimental/NNERuntimeRDG/Source/NNERuntimeRDG/Private/Hlsl/NNERuntimeRDGBatchNormalization.cpp
@@ -25,19 +25,19 @@ namespace UE::NNERuntimeRDG::Private::Hlsl
 
 	public:
 
-		virtual int PrepareOutputs(TConstArrayView<NNECore::Internal::FTensorRef> InputTensors, TArrayView<NNECore::Internal::FTensorRef> OutputTensors) const override
+		virtual int PrepareOutputs(TConstArrayView<NNE::Internal::FTensorRef> InputTensors, TArrayView<NNE::Internal::FTensorRef> OutputTensors) const override
 		{
 			check(InputTensors.Num() == 5);
 			check(OutputTensors.Num() == 1);
 
-			const NNECore::Internal::FTensor& X = *InputTensors[0];
+			const NNE::Internal::FTensor& X = *InputTensors[0];
 
 			OutputTensors[0]->SetShape(X.GetShape());
 
 			return 0;
 		};
 
-		virtual bool Initialize(TConstArrayView<NNECore::FTensorDesc> InputTensorDescs, TConstArrayView<NNECore::FTensorDesc> OutputTensorDescs, const NNECore::FAttributeMap& Attributes) override
+		virtual bool Initialize(TConstArrayView<NNE::FTensorDesc> InputTensorDescs, TConstArrayView<NNE::FTensorDesc> OutputTensorDescs, const NNE::FAttributeMap& Attributes) override
 		{
 			check(InputTensorDescs.Num() == 5);
 			check(OutputTensorDescs.Num() >= 1);
@@ -122,7 +122,7 @@ namespace UE::NNERuntimeRDG::Private::Hlsl
 		}
 	};
 
-	bool ValidateBatchNormalizationOperator(const NNECore::FAttributeMap& AttributeMap, TConstArrayView<ENNETensorDataType> InputTypes, TConstArrayView<NNECore::FSymbolicTensorShape> InputBatchNormalizations)
+	bool ValidateBatchNormalizationOperator(const NNE::FAttributeMap& AttributeMap, TConstArrayView<ENNETensorDataType> InputTypes, TConstArrayView<NNE::FSymbolicTensorShape> InputBatchNormalizations)
 	{
 		bool bIsValid = true;
 
diff --git a/Engine/Plugins/Experimental/NNERuntimeRDG/Source/NNERuntimeRDG/Private/Hlsl/NNERuntimeRDGCast.cpp b/Engine/Plugins/Experimental/NNERuntimeRDG/Source/NNERuntimeRDG/Private/Hlsl/NNERuntimeRDGCast.cpp
--- a/Engine/Plugins/Experimental/NNERuntimeRDG/Source/NNERuntimeRDG/Private/Hlsl/NNERuntimeRDGCast.cpp
+++ b/Engine/Plugins/Experimental/NNERuntimeRDG/Source/NNERuntimeRDG/Private/Hlsl/NNERuntimeRDGCast.cpp
@@ -18,13 +18,13 @@ namespace UE::NNERuntimeRDG::Private::Hlsl
 		FCast() {}
 		virtual ~FCast() = default;
 
-		virtual int PrepareOutputs(TConstArrayView<NNECore::Internal::FTensorRef> InputTensors, TArrayView<NNECore::Internal::FTensorRef> OutputTensors) const override
+		virtual int PrepareOutputs(TConstArrayView<NNE::Internal::FTensorRef> InputTensors, TArrayView<NNE::Internal::FTensorRef> OutputTensors) const override
 		{
 			check(InputTensors.Num() == 1);
 			check(OutputTensors.Num() == 1);
 			OutputTensors[0]->SetShape(InputTensors[0]->GetShape());
 
-			const NNECore::Internal::FTensor& X = *InputTensors[0];
+			const NNE::Internal::FTensor& X = *InputTensors[0];
 
 			Internal::CPUHelper::Cast::Apply(X, *OutputTensors[0]);
 
@@ -37,7 +37,7 @@ namespace UE::NNERuntimeRDG::Private::Hlsl
 			return 0;
 		}
 
-		virtual bool Initialize(TConstArrayView<NNECore::FTensorDesc> InputTensorDescs, TConstArrayView<NNECore::FTensorDesc> OutputTensorDescs, const NNECore::FAttributeMap& Attributes) override
+		virtual bool Initialize(TConstArrayView<NNE::FTensorDesc> InputTensorDescs, TConstArrayView<NNE::FTensorDesc> OutputTensorDescs, const NNE::FAttributeMap& Attributes) override
 		{
 			check(InputTensorDescs.Num() == 1);
 			check(OutputTensorDescs.Num() == 1);
@@ -60,7 +60,7 @@ namespace UE::NNERuntimeRDG::Private::Hlsl
 		}
 	};
 
-	bool ValidateCastOperator(const NNECore::FAttributeMap& AttributeMap, TConstArrayView<ENNETensorDataType> InputTypes, TConstArrayView<NNECore::FSymbolicTensorShape> InputShapes)
+	bool ValidateCastOperator(const NNE::FAttributeMap& AttributeMap, TConstArrayView<ENNETensorDataType> InputTypes, TConstArrayView<NNE::FSymbolicTensorShape> InputShapes)
 	{
 		bool bIsValid = true;
 
diff --git a/Engine/Plugins/Experimental/NNERuntimeRDG/Source/NNERuntimeRDG/Private/Hlsl/NNERuntimeRDGConcat.cpp b/Engine/Plugins/Experimental/NNERuntimeRDG/Source/NNERuntimeRDG/Private/Hlsl/NNERuntimeRDGConcat.cpp
--- a/Engine/Plugins/Experimental/NNERuntimeRDG/Source/NNERuntimeRDG/Private/Hlsl/NNERuntimeRDGConcat.cpp
+++ b/Engine/Plugins/Experimental/NNERuntimeRDG/Source/NNERuntimeRDG/Private/Hlsl/NNERuntimeRDGConcat.cpp
@@ -21,7 +21,7 @@ namespace UE::NNERuntimeRDG::Private::Hlsl
 
 	public:
 
-		virtual int PrepareOutputs(TConstArrayView<NNECore::Internal::FTensorRef> InputTensors, TArrayView<NNECore::Internal::FTensorRef> OutputTensors) const override
+		virtual int PrepareOutputs(TConstArrayView<NNE::Internal::FTensorRef> InputTensors, TArrayView<NNE::Internal::FTensorRef> OutputTensors) const override
 		{
 			check(InputTensors.Num() >= 1);
 			check(OutputTensors.Num() == 1);
@@ -43,7 +43,7 @@ namespace UE::NNERuntimeRDG::Private::Hlsl
 				}
 			}
 
-			NNECore::FTensorShape OutputShape = NNECore::FTensorShape::Make(OutputShapeData);
+			NNE::FTensorShape OutputShape = NNE::FTensorShape::Make(OutputShapeData);
 			
 			OutputTensors[0]->SetShape(OutputShape);
 
@@ -58,7 +58,7 @@ namespace UE::NNERuntimeRDG::Private::Hlsl
 			return 0;
 		};
 
-		virtual bool Initialize(TConstArrayView<NNECore::FTensorDesc> InputTensorDescs, TConstArrayView<NNECore::FTensorDesc> OutputTensorDescs, const NNECore::FAttributeMap& Attributes) override
+		virtual bool Initialize(TConstArrayView<NNE::FTensorDesc> InputTensorDescs, TConstArrayView<NNE::FTensorDesc> OutputTensorDescs, const NNE::FAttributeMap& Attributes) override
 		{
 			check(InputTensorDescs.Num() >= 1);
 			check(OutputTensorDescs.Num() == 1);
@@ -91,7 +91,7 @@ namespace UE::NNERuntimeRDG::Private::Hlsl
 		}
 	};
 
-	bool ValidateConcatOperator(const NNECore::FAttributeMap& AttributeMap, TConstArrayView<ENNETensorDataType> InputTypes, TConstArrayView<NNECore::FSymbolicTensorShape> InputConcats)
+	bool ValidateConcatOperator(const NNE::FAttributeMap& AttributeMap, TConstArrayView<ENNETensorDataType> InputTypes, TConstArrayView<NNE::FSymbolicTensorShape> InputConcats)
 	{
 		bool bIsValid = true;
 
diff --git a/Engine/Plugins/Experimental/NNERuntimeRDG/Source/NNERuntimeRDG/Private/Hlsl/NNERuntimeRDGConv.cpp b/Engine/Plugins/Experimental/NNERuntimeRDG/Source/NNERuntimeRDG/Private/Hlsl/NNERuntimeRDGConv.cpp
--- a/Engine/Plugins/Experimental/NNERuntimeRDG/Source/NNERuntimeRDG/Private/Hlsl/NNERuntimeRDGConv.cpp
+++ b/Engine/Plugins/Experimental/NNERuntimeRDG/Source/NNERuntimeRDG/Private/Hlsl/NNERuntimeRDGConv.cpp
@@ -46,36 +46,36 @@ namespace UE::NNERuntimeRDG::Private::Hlsl
 
 	public:
 
-		virtual int PrepareOutputs(TConstArrayView<NNECore::Internal::FTensorRef> InputTensors, TArrayView<NNECore::Internal::FTensorRef> OutputTensors) const override
+		virtual int PrepareOutputs(TConstArrayView<NNE::Internal::FTensorRef> InputTensors, TArrayView<NNE::Internal::FTensorRef> OutputTensors) const override
 		{
 			check(InputTensors.Num() >= 2 && InputTensors.Num() <= 3);
 			check(OutputTensors.Num() == 1);
 
-			const NNECore::FTensorShape& Input = InputTensors[0]->GetShape();
-			const NNECore::FTensorShape& Weights = InputTensors[1]->GetShape();
+			const NNE::FTensorShape& Input = InputTensors[0]->GetShape();
+			const NNE::FTensorShape& Weights = InputTensors[1]->GetShape();
 
 			const TArray<int32> OutputShapeData = NNEHlslShaders::Internal::FConvCS::GetOutputShape(Input.GetData(), Weights.GetData(), AutoPad, Dilations, Strides, Pads);
-			const NNECore::FSymbolicTensorShape OutputShape = NNECore::FSymbolicTensorShape::Make(OutputShapeData);
+			const NNE::FSymbolicTensorShape OutputShape = NNE::FSymbolicTensorShape::Make(OutputShapeData);
 			
 			if (!OutputShape.IsConcrete())
 			{
 				return -1;
 			}
-			OutputTensors[0]->SetShape(NNECore::FTensorShape::MakeFromSymbolic(OutputShape));
+			OutputTensors[0]->SetShape(NNE::FTensorShape::MakeFromSymbolic(OutputShape));
 
 			return 0;
 		};
 
-		virtual bool Initialize(TConstArrayView<NNECore::FTensorDesc> InputTensorDescs, TConstArrayView<NNECore::FTensorDesc> OutputTensorDescs, const NNECore::FAttributeMap& Attributes) override
+		virtual bool Initialize(TConstArrayView<NNE::FTensorDesc> InputTensorDescs, TConstArrayView<NNE::FTensorDesc> OutputTensorDescs, const NNE::FAttributeMap& Attributes) override
 		{
             using namespace UE::NNEHlslShaders::Internal;
 
 			check(InputTensorDescs.Num() >= 2 && InputTensorDescs.Num() <= 3);
 			check(OutputTensorDescs.Num() == 1);
 
-            const NNECore::FTensorDesc& Input = InputTensorDescs[0];
-			const NNECore::FTensorDesc& Weights = InputTensorDescs[1];
-			const NNECore::FTensorDesc& Output = OutputTensorDescs[0];
+            const NNE::FTensorDesc& Input = InputTensorDescs[0];
+			const NNE::FTensorDesc& Weights = InputTensorDescs[1];
+			const NNE::FTensorDesc& Output = OutputTensorDescs[0];
 			
 			if (Input.GetShape().Rank() < 2)
 			{
@@ -298,7 +298,7 @@ namespace UE::NNERuntimeRDG::Private::Hlsl
 		};
 	};
 
-	bool ValidateConvOperator(const NNECore::FAttributeMap& AttributeMap, TConstArrayView<ENNETensorDataType> InputTypes, TConstArrayView<NNECore::FSymbolicTensorShape> InputShapes)
+	bool ValidateConvOperator(const NNE::FAttributeMap& AttributeMap, TConstArrayView<ENNETensorDataType> InputTypes, TConstArrayView<NNE::FSymbolicTensorShape> InputShapes)
 	{
 		bool bIsValid = true;
 
diff --git a/Engine/Plugins/Experimental/NNERuntimeRDG/Source/NNERuntimeRDG/Private/Hlsl/NNERuntimeRDGConvTranspose.cpp b/Engine/Plugins/Experimental/NNERuntimeRDG/Source/NNERuntimeRDG/Private/Hlsl/NNERuntimeRDGConvTranspose.cpp
--- a/Engine/Plugins/Experimental/NNERuntimeRDG/Source/NNERuntimeRDG/Private/Hlsl/NNERuntimeRDGConvTranspose.cpp
+++ b/Engine/Plugins/Experimental/NNERuntimeRDG/Source/NNERuntimeRDG/Private/Hlsl/NNERuntimeRDGConvTranspose.cpp
@@ -41,36 +41,36 @@ namespace UE::NNERuntimeRDG::Private::Hlsl
 
 	public:
 
-		virtual int PrepareOutputs(TConstArrayView<NNECore::Internal::FTensorRef> InputTensors, TArrayView<NNECore::Internal::FTensorRef> OutputTensors) const override
+		virtual int PrepareOutputs(TConstArrayView<NNE::Internal::FTensorRef> InputTensors, TArrayView<NNE::Internal::FTensorRef> OutputTensors) const override
 		{
 			check(InputTensors.Num() >= 2 && InputTensors.Num() <= 3);
 			check(OutputTensors.Num() == 1);
 			
-			const NNECore::FTensorShape& Input = InputTensors[0]->GetShape();
-			const NNECore::FTensorShape& Weights = InputTensors[1]->GetShape();
+			const NNE::FTensorShape& Input = InputTensors[0]->GetShape();
+			const NNE::FTensorShape& Weights = InputTensors[1]->GetShape();
 
 			
 			const TArray<int32> OutputShapeData = NNEHlslShaders::Internal::FConvTransposeCS::GetOutputShape(Input.GetData(), Weights.GetData(), AutoPad, Dilations, Strides, Pads, OutputPadding, Group);
-			const NNECore::FSymbolicTensorShape OutputShape = NNECore::FSymbolicTensorShape::Make(OutputShapeData);
+			const NNE::FSymbolicTensorShape OutputShape = NNE::FSymbolicTensorShape::Make(OutputShapeData);
 			if (!OutputShape.IsConcrete())
 			{
 				return -1;
 			}
-			OutputTensors[0]->SetShape(NNECore::FTensorShape::MakeFromSymbolic(OutputShape));
+			OutputTensors[0]->SetShape(NNE::FTensorShape::MakeFromSymbolic(OutputShape));
 
 			return 0;
 		};
 
-		virtual bool Initialize(TConstArrayView<NNECore::FTensorDesc> InputTensorDescs, TConstArrayView<NNECore::FTensorDesc> OutputTensorDescs, const NNECore::FAttributeMap& Attributes) override
+		virtual bool Initialize(TConstArrayView<NNE::FTensorDesc> InputTensorDescs, TConstArrayView<NNE::FTensorDesc> OutputTensorDescs, const NNE::FAttributeMap& Attributes) override
 		{
 			using namespace UE::NNEHlslShaders::Internal;
 
 			check(InputTensorDescs.Num() >= 2 && InputTensorDescs.Num() <= 3);
 			check(OutputTensorDescs.Num() == 1);
 
-            const NNECore::FTensorDesc& Input = InputTensorDescs[0];
-			const NNECore::FTensorDesc& Weights = InputTensorDescs[1];
-			const NNECore::FTensorDesc& Output = OutputTensorDescs[0];
+            const NNE::FTensorDesc& Input = InputTensorDescs[0];
+			const NNE::FTensorDesc& Weights = InputTensorDescs[1];
+			const NNE::FTensorDesc& Output = OutputTensorDescs[0];
 
 			if (Input.GetShape().Rank() < 2)
 			{
@@ -177,7 +177,7 @@ namespace UE::NNERuntimeRDG::Private::Hlsl
 		}
 	};
 
-	bool ValidateConvTransposeOperator(const NNECore::FAttributeMap& AttributeMap, TConstArrayView<ENNETensorDataType> InputTypes, TConstArrayView<NNECore::FSymbolicTensorShape> InputShapes)
+	bool ValidateConvTransposeOperator(const NNE::FAttributeMap& AttributeMap, TConstArrayView<ENNETensorDataType> InputTypes, TConstArrayView<NNE::FSymbolicTensorShape> InputShapes)
 	{
 		bool bIsValid = true;
 
diff --git a/Engine/Plugins/Experimental/NNERuntimeRDG/Source/NNERuntimeRDG/Private/Hlsl/NNERuntimeRDGDropout.cpp b/Engine/Plugins/Experimental/NNERuntimeRDG/Source/NNERuntimeRDG/Private/Hlsl/NNERuntimeRDGDropout.cpp
--- a/Engine/Plugins/Experimental/NNERuntimeRDG/Source/NNERuntimeRDG/Private/Hlsl/NNERuntimeRDGDropout.cpp
+++ b/Engine/Plugins/Experimental/NNERuntimeRDG/Source/NNERuntimeRDG/Private/Hlsl/NNERuntimeRDGDropout.cpp
@@ -21,12 +21,12 @@ namespace UE::NNERuntimeRDG::Private::Hlsl
 
 	public:
 
-		virtual int PrepareOutputs(TConstArrayView<NNECore::Internal::FTensorRef> InputTensors, TArrayView<NNECore::Internal::FTensorRef> OutputTensors) const override
+		virtual int PrepareOutputs(TConstArrayView<NNE::Internal::FTensorRef> InputTensors, TArrayView<NNE::Internal::FTensorRef> OutputTensors) const override
 		{
 			check(InputTensors.Num() == 1);
 			check(OutputTensors.Num() == 1);
 
-			const NNECore::Internal::FTensor& X = *InputTensors[0];
+			const NNE::Internal::FTensor& X = *InputTensors[0];
 
 			OutputTensors[0]->SetShape(X.GetShape());
 			if (X.HasPreparedData())
@@ -37,7 +37,7 @@ namespace UE::NNERuntimeRDG::Private::Hlsl
 			return 0;
 		};
 
-		virtual bool Initialize(TConstArrayView<NNECore::FTensorDesc> InputTensorDescs, TConstArrayView<NNECore::FTensorDesc> OutputTensorDescs, const NNECore::FAttributeMap& Attributes) override
+		virtual bool Initialize(TConstArrayView<NNE::FTensorDesc> InputTensorDescs, TConstArrayView<NNE::FTensorDesc> OutputTensorDescs, const NNE::FAttributeMap& Attributes) override
 		{
 			check(InputTensorDescs.Num() == 1);
 			check(OutputTensorDescs.Num() == 1 || OutputTensorDescs.Num() == 2);
@@ -68,7 +68,7 @@ namespace UE::NNERuntimeRDG::Private::Hlsl
 		}
 	};
 
-	bool ValidateDropoutOperator(const NNECore::FAttributeMap& AttributeMap, TConstArrayView<ENNETensorDataType> InputTypes, TConstArrayView<NNECore::FSymbolicTensorShape> InputDropouts)
+	bool ValidateDropoutOperator(const NNE::FAttributeMap& AttributeMap, TConstArrayView<ENNETensorDataType> InputTypes, TConstArrayView<NNE::FSymbolicTensorShape> InputDropouts)
 	{
 		bool bIsValid = true;
 
diff --git a/Engine/Plugins/Experimental/NNERuntimeRDG/Source/NNERuntimeRDG/Private/Hlsl/NNERuntimeRDGElementWiseBinary.cpp b/Engine/Plugins/Experimental/NNERuntimeRDG/Source/NNERuntimeRDG/Private/Hlsl/NNERuntimeRDGElementWiseBinary.cpp
--- a/Engine/Plugins/Experimental/NNERuntimeRDG/Source/NNERuntimeRDG/Private/Hlsl/NNERuntimeRDGElementWiseBinary.cpp
+++ b/Engine/Plugins/Experimental/NNERuntimeRDG/Source/NNERuntimeRDG/Private/Hlsl/NNERuntimeRDGElementWiseBinary.cpp
@@ -17,7 +17,7 @@ namespace UE::NNERuntimeRDG::Private::Hlsl
 	/**
 	 * Binary element-wise operator implementation
 	 */
-	template<NNECore::Internal::EElementWiseBinaryOperatorType OpType>
+	template<NNE::Internal::EElementWiseBinaryOperatorType OpType>
 	class TElementWiseBinary : public FOperatorHlsl
 	{
 	public:
@@ -27,13 +27,13 @@ namespace UE::NNERuntimeRDG::Private::Hlsl
 
 	public:
 
-		virtual int PrepareOutputs(TConstArrayView<NNECore::Internal::FTensorRef> InputTensors, TArrayView<NNECore::Internal::FTensorRef> OutputTensors) const override
+		virtual int PrepareOutputs(TConstArrayView<NNE::Internal::FTensorRef> InputTensors, TArrayView<NNE::Internal::FTensorRef> OutputTensors) const override
 		{
 			check(InputTensors.Num() == 2);
 			check(OutputTensors.Num() == 1);
 			
-			const NNECore::FTensorShape& LHSInput = InputTensors[0]->GetShape();
-			const NNECore::FTensorShape& RHSInput = InputTensors[1]->GetShape();
+			const NNE::FTensorShape& LHSInput = InputTensors[0]->GetShape();
+			const NNE::FTensorShape& RHSInput = InputTensors[1]->GetShape();
 			const int32 OutputRank = FMath::Max(LHSInput.Rank(), RHSInput.Rank());
 			TArray<uint32> OutputShapeData;
 			
@@ -54,7 +54,7 @@ namespace UE::NNERuntimeRDG::Private::Hlsl
 				OutputShapeData[OutputRank - 1 - i] = OutputValue;
 			}
 
-			NNECore::FTensorShape OutputShape = NNECore::FTensorShape::Make(OutputShapeData);
+			NNE::FTensorShape OutputShape = NNE::FTensorShape::Make(OutputShapeData);
 			
 			OutputTensors[0]->SetShape(OutputShape);
 
@@ -69,7 +69,7 @@ namespace UE::NNERuntimeRDG::Private::Hlsl
 			return 0;
 		}
 		
-		virtual bool Initialize(TConstArrayView<NNECore::FTensorDesc> InputTensorDescs, TConstArrayView<NNECore::FTensorDesc> OutputTensorDescs, const NNECore::FAttributeMap& Attributes) override
+		virtual bool Initialize(TConstArrayView<NNE::FTensorDesc> InputTensorDescs, TConstArrayView<NNE::FTensorDesc> OutputTensorDescs, const NNE::FAttributeMap& Attributes) override
 		{
 			check(InputTensorDescs.Num() == 2);
 			check(OutputTensorDescs.Num() == 1);
@@ -127,7 +127,7 @@ namespace UE::NNERuntimeRDG::Private::Hlsl
 		}
 	};
 
-	bool ValidateElementWiseBinaryOperator(const NNECore::FAttributeMap& AttributeMap, TConstArrayView<ENNETensorDataType> InputTypes, TConstArrayView<NNECore::FSymbolicTensorShape> InputShapes)
+	bool ValidateElementWiseBinaryOperator(const NNE::FAttributeMap& AttributeMap, TConstArrayView<ENNETensorDataType> InputTypes, TConstArrayView<NNE::FSymbolicTensorShape> InputShapes)
 	{
 		bool bIsValid = true;
 
@@ -145,7 +145,7 @@ namespace UE::NNERuntimeRDG::Private::Hlsl
 		return bIsValid;
 	}
 
-	template<NNECore::Internal::EElementWiseBinaryOperatorType OpType>
+	template<NNE::Internal::EElementWiseBinaryOperatorType OpType>
 	FOperatorHlsl* CreateElementWiseBinaryOperator()
 	{
 		return new TElementWiseBinary<OpType>();
@@ -153,7 +153,7 @@ namespace UE::NNERuntimeRDG::Private::Hlsl
 
 	bool RegisterElementWiseBinaryOperators(FOperatorRegistryHlsl& Registry)
 	{
-#define OP(Name) Registry.OpAdd(TEXT(#Name), CreateElementWiseBinaryOperator<NNECore::Internal::EElementWiseBinaryOperatorType::Name>, ValidateElementWiseBinaryOperator)
+#define OP(Name) Registry.OpAdd(TEXT(#Name), CreateElementWiseBinaryOperator<NNE::Internal::EElementWiseBinaryOperatorType::Name>, ValidateElementWiseBinaryOperator)
 		OP(Add);
 		//OP(And);
 		OP(Div);
diff --git a/Engine/Plugins/Experimental/NNERuntimeRDG/Source/NNERuntimeRDG/Private/Hlsl/NNERuntimeRDGElementWiseUnary.cpp b/Engine/Plugins/Experimental/NNERuntimeRDG/Source/NNERuntimeRDG/Private/Hlsl/NNERuntimeRDGElementWiseUnary.cpp
--- a/Engine/Plugins/Experimental/NNERuntimeRDG/Source/NNERuntimeRDG/Private/Hlsl/NNERuntimeRDGElementWiseUnary.cpp
+++ b/Engine/Plugins/Experimental/NNERuntimeRDG/Source/NNERuntimeRDG/Private/Hlsl/NNERuntimeRDGElementWiseUnary.cpp
@@ -18,7 +18,7 @@ namespace UE::NNERuntimeRDG::Private::Hlsl
 	/**
 	 * Unary element-wise operator implementation
 	 */
-	template<NNECore::Internal::EElementWiseUnaryOperatorType OpType>
+	template<NNE::Internal::EElementWiseUnaryOperatorType OpType>
 	class TElementWiseUnary : public FOperatorHlsl
 	{
 	public:
@@ -34,20 +34,20 @@ namespace UE::NNERuntimeRDG::Private::Hlsl
 
 	public:
 
-		virtual int PrepareOutputs(TConstArrayView<NNECore::Internal::FTensorRef> InputTensors, TArrayView<NNECore::Internal::FTensorRef> OutputTensors) const override
+		virtual int PrepareOutputs(TConstArrayView<NNE::Internal::FTensorRef> InputTensors, TArrayView<NNE::Internal::FTensorRef> OutputTensors) const override
 		{
 			check(InputTensors.Num() == 1);
 			check(OutputTensors.Num() == 1);
 			OutputTensors[0]->SetShape(InputTensors[0]->GetShape());
 
-			const NNECore::Internal::FTensor& X = *InputTensors[0];
+			const NNE::Internal::FTensor& X = *InputTensors[0];
 
 			Internal::CPUHelper::ElementWiseUnary::Apply(OpType, X, Alpha, Beta, Gamma, *OutputTensors[0]);
 			
 			return 0;
 		}
 
-		virtual bool Initialize(TConstArrayView<NNECore::FTensorDesc> InputTensorDescs, TConstArrayView<NNECore::FTensorDesc> OutputTensorDescs, const NNECore::FAttributeMap& Attributes) override
+		virtual bool Initialize(TConstArrayView<NNE::FTensorDesc> InputTensorDescs, TConstArrayView<NNE::FTensorDesc> OutputTensorDescs, const NNE::FAttributeMap& Attributes) override
 		{
 			check(InputTensorDescs.Num() == 1);
 			check(OutputTensorDescs.Num() == 1);
@@ -97,27 +97,27 @@ namespace UE::NNERuntimeRDG::Private::Hlsl
 		}
 	};
 
-	template<> TElementWiseUnary<NNECore::Internal::EElementWiseUnaryOperatorType::Selu>::TElementWiseUnary()
+	template<> TElementWiseUnary<NNE::Internal::EElementWiseUnaryOperatorType::Selu>::TElementWiseUnary()
 		: Alpha(1.67326319217681884765625f), Beta(0.0f), Gamma(1.05070102214813232421875f)
 	{
 	}
 
-	template<> TElementWiseUnary<NNECore::Internal::EElementWiseUnaryOperatorType::Elu>::TElementWiseUnary()
+	template<> TElementWiseUnary<NNE::Internal::EElementWiseUnaryOperatorType::Elu>::TElementWiseUnary()
 		: Alpha(1.0f), Beta(0.0f), Gamma(0.0f) 
 	{
 	}
 
-	template<> TElementWiseUnary<NNECore::Internal::EElementWiseUnaryOperatorType::HardSigmoid>::TElementWiseUnary()
+	template<> TElementWiseUnary<NNE::Internal::EElementWiseUnaryOperatorType::HardSigmoid>::TElementWiseUnary()
 		: Alpha(0.2f), Beta(0.5f), Gamma(0.0f)
 	{
 	}
 
-	template<> TElementWiseUnary<NNECore::Internal::EElementWiseUnaryOperatorType::LeakyRelu>::TElementWiseUnary()
+	template<> TElementWiseUnary<NNE::Internal::EElementWiseUnaryOperatorType::LeakyRelu>::TElementWiseUnary()
 		: Alpha(0.01f), Beta(0.0f), Gamma(0.0f)
 	{
 	}
 
-	template<> bool TElementWiseUnary<NNECore::Internal::EElementWiseUnaryOperatorType::Clip>::Initialize(TConstArrayView<NNECore::FTensorDesc> InputTensorDescs, TConstArrayView<NNECore::FTensorDesc> OutputTensorDescs, const NNECore::FAttributeMap& Attributes)
+	template<> bool TElementWiseUnary<NNE::Internal::EElementWiseUnaryOperatorType::Clip>::Initialize(TConstArrayView<NNE::FTensorDesc> InputTensorDescs, TConstArrayView<NNE::FTensorDesc> OutputTensorDescs, const NNE::FAttributeMap& Attributes)
 	{
 		check(InputTensorDescs.Num() == 1);
 		check(OutputTensorDescs.Num() == 1);
@@ -127,14 +127,14 @@ namespace UE::NNERuntimeRDG::Private::Hlsl
 		return true;
 	}
 
-	template<NNECore::Internal::EElementWiseUnaryOperatorType OpType>
+	template<NNE::Internal::EElementWiseUnaryOperatorType OpType>
 	FOperatorHlsl* CreateElementWiseUnaryOperator()
 	{
 		return new TElementWiseUnary<OpType>();
 	}
 
-	template<NNECore::Internal::EElementWiseUnaryOperatorType OpType>
-	bool ValidateElementWiseUnaryOperator(const NNECore::FAttributeMap& AttributeMap, TConstArrayView<ENNETensorDataType> InputTypes, TConstArrayView<NNECore::FSymbolicTensorShape> InputShapes)
+	template<NNE::Internal::EElementWiseUnaryOperatorType OpType>
+	bool ValidateElementWiseUnaryOperator(const NNE::FAttributeMap& AttributeMap, TConstArrayView<ENNETensorDataType> InputTypes, TConstArrayView<NNE::FSymbolicTensorShape> InputShapes)
 	{
 		bool bIsValid = true;
 
@@ -150,7 +150,7 @@ namespace UE::NNERuntimeRDG::Private::Hlsl
 	}
 
 	template<>
-	bool ValidateElementWiseUnaryOperator<NNECore::Internal::EElementWiseUnaryOperatorType::Selu>(const NNECore::FAttributeMap& AttributeMap, TConstArrayView<ENNETensorDataType> InputTypes, TConstArrayView<NNECore::FSymbolicTensorShape> InputShapes)
+	bool ValidateElementWiseUnaryOperator<NNE::Internal::EElementWiseUnaryOperatorType::Selu>(const NNE::FAttributeMap& AttributeMap, TConstArrayView<ENNETensorDataType> InputTypes, TConstArrayView<NNE::FSymbolicTensorShape> InputShapes)
 	{
 		bool bIsValid = true;
 
@@ -168,7 +168,7 @@ namespace UE::NNERuntimeRDG::Private::Hlsl
 	}
 
 	template<>
-	bool ValidateElementWiseUnaryOperator<NNECore::Internal::EElementWiseUnaryOperatorType::Elu>(const NNECore::FAttributeMap& AttributeMap, TConstArrayView<ENNETensorDataType> InputTypes, TConstArrayView<NNECore::FSymbolicTensorShape> InputShapes)
+	bool ValidateElementWiseUnaryOperator<NNE::Internal::EElementWiseUnaryOperatorType::Elu>(const NNE::FAttributeMap& AttributeMap, TConstArrayView<ENNETensorDataType> InputTypes, TConstArrayView<NNE::FSymbolicTensorShape> InputShapes)
 	{
 		bool bIsValid = true;
 
@@ -185,7 +185,7 @@ namespace UE::NNERuntimeRDG::Private::Hlsl
 	}
 
 	template<>
-	bool ValidateElementWiseUnaryOperator<NNECore::Internal::EElementWiseUnaryOperatorType::HardSigmoid>(const NNECore::FAttributeMap& AttributeMap, TConstArrayView<ENNETensorDataType> InputTypes, TConstArrayView<NNECore::FSymbolicTensorShape> InputShapes)
+	bool ValidateElementWiseUnaryOperator<NNE::Internal::EElementWiseUnaryOperatorType::HardSigmoid>(const NNE::FAttributeMap& AttributeMap, TConstArrayView<ENNETensorDataType> InputTypes, TConstArrayView<NNE::FSymbolicTensorShape> InputShapes)
 	{
 		bool bIsValid = true;
 
@@ -203,7 +203,7 @@ namespace UE::NNERuntimeRDG::Private::Hlsl
 	}
 
 	template<>
-	bool ValidateElementWiseUnaryOperator<NNECore::Internal::EElementWiseUnaryOperatorType::LeakyRelu>(const NNECore::FAttributeMap& AttributeMap, TConstArrayView<ENNETensorDataType> InputTypes, TConstArrayView<NNECore::FSymbolicTensorShape> InputShapes)
+	bool ValidateElementWiseUnaryOperator<NNE::Internal::EElementWiseUnaryOperatorType::LeakyRelu>(const NNE::FAttributeMap& AttributeMap, TConstArrayView<ENNETensorDataType> InputTypes, TConstArrayView<NNE::FSymbolicTensorShape> InputShapes)
 	{
 		bool bIsValid = true;
 
@@ -220,7 +220,7 @@ namespace UE::NNERuntimeRDG::Private::Hlsl
 	}
 
 	template<>
-	bool ValidateElementWiseUnaryOperator<NNECore::Internal::EElementWiseUnaryOperatorType::Clip>(const NNECore::FAttributeMap& AttributeMap, TConstArrayView<ENNETensorDataType> InputTypes, TConstArrayView<NNECore::FSymbolicTensorShape> InputShapes)
+	bool ValidateElementWiseUnaryOperator<NNE::Internal::EElementWiseUnaryOperatorType::Clip>(const NNE::FAttributeMap& AttributeMap, TConstArrayView<ENNETensorDataType> InputTypes, TConstArrayView<NNE::FSymbolicTensorShape> InputShapes)
 	{
 		bool bIsValid = true;
 
@@ -241,7 +241,7 @@ namespace UE::NNERuntimeRDG::Private::Hlsl
 	
 	bool RegisterElementWiseUnaryOperators(FOperatorRegistryHlsl& Registry)
 	{
-#define OP(Name) Registry.OpAdd(TEXT(#Name), CreateElementWiseUnaryOperator<NNECore::Internal::EElementWiseUnaryOperatorType::Name>, ValidateElementWiseUnaryOperator<NNECore::Internal::EElementWiseUnaryOperatorType::Name>)
+#define OP(Name) Registry.OpAdd(TEXT(#Name), CreateElementWiseUnaryOperator<NNE::Internal::EElementWiseUnaryOperatorType::Name>, ValidateElementWiseUnaryOperator<NNE::Internal::EElementWiseUnaryOperatorType::Name>)
 		OP(Abs);
 		OP(Acos);
 		OP(Acosh);
diff --git a/Engine/Plugins/Experimental/NNERuntimeRDG/Source/NNERuntimeRDG/Private/Hlsl/NNERuntimeRDGElementWiseVariadic.cpp b/Engine/Plugins/Experimental/NNERuntimeRDG/Source/NNERuntimeRDG/Private/Hlsl/NNERuntimeRDGElementWiseVariadic.cpp
--- a/Engine/Plugins/Experimental/NNERuntimeRDG/Source/NNERuntimeRDG/Private/Hlsl/NNERuntimeRDGElementWiseVariadic.cpp
+++ b/Engine/Plugins/Experimental/NNERuntimeRDG/Source/NNERuntimeRDG/Private/Hlsl/NNERuntimeRDGElementWiseVariadic.cpp
@@ -17,7 +17,7 @@ namespace UE::NNERuntimeRDG::Private::Hlsl
 		TConstArrayView<FTensorRDGRef> InputTensors,
 		const FTensorRDG& OutputTensor,
 		bool OutputAsInput,
-		NNECore::Internal::EElementWiseVariadicOperatorType OpType,
+		NNE::Internal::EElementWiseVariadicOperatorType OpType,
 		float Scale)
 	{
 		static_assert(FElementWiseVariadicConstants::MAX_NUM_INPUT == 4, "This algorithm need to be adapted to math the shader.");
@@ -84,7 +84,7 @@ namespace UE::NNERuntimeRDG::Private::Hlsl
 	/**
 	 * Variadic Element-wise ML operator implementation
 	 */
-	template<NNECore::Internal::EElementWiseVariadicOperatorType OpType>
+	template<NNE::Internal::EElementWiseVariadicOperatorType OpType>
 	class TElementWiseVariadic : public FOperatorHlsl
 	{
 	public:
@@ -94,7 +94,7 @@ namespace UE::NNERuntimeRDG::Private::Hlsl
 
 	public:
 
-		virtual int PrepareOutputs(TConstArrayView<NNECore::Internal::FTensorRef> InputTensors, TArrayView<NNECore::Internal::FTensorRef> OutputTensors) const override
+		virtual int PrepareOutputs(TConstArrayView<NNE::Internal::FTensorRef> InputTensors, TArrayView<NNE::Internal::FTensorRef> OutputTensors) const override
 		{
 			check(InputTensors.Num() > 0);
 			check(OutputTensors.Num() == 1);
@@ -127,14 +127,14 @@ namespace UE::NNERuntimeRDG::Private::Hlsl
 				OutputShapeData[OutputRank - 1 - i] = OutputValue;
 			}
 
-			NNECore::FTensorShape OutputShape = NNECore::FTensorShape::Make(OutputShapeData);
+			NNE::FTensorShape OutputShape = NNE::FTensorShape::Make(OutputShapeData);
 
 			OutputTensors[0]->SetShape(OutputShape);
 
 			return 0;
 		}
 		
-		virtual bool Initialize(TConstArrayView<NNECore::FTensorDesc> InputTensorDescs, TConstArrayView<NNECore::FTensorDesc> OutputTensorDescs, const NNECore::FAttributeMap& Attributes) override
+		virtual bool Initialize(TConstArrayView<NNE::FTensorDesc> InputTensorDescs, TConstArrayView<NNE::FTensorDesc> OutputTensorDescs, const NNE::FAttributeMap& Attributes) override
 		{
 			check(InputTensorDescs.Num() > 0);
 			check(OutputTensorDescs.Num() == 1);
@@ -165,7 +165,7 @@ namespace UE::NNERuntimeRDG::Private::Hlsl
 					PassInputTensors[i] = InInputTensors[InputOffset + i];
 				}
 
-				if (OpType == NNECore::Internal::EElementWiseVariadicOperatorType::Mean && bIsLastPass)
+				if (OpType == NNE::Internal::EElementWiseVariadicOperatorType::Mean && bIsLastPass)
 				{
 					Scale = 1.0f / InInputTensors.Num();
 				}
@@ -179,7 +179,7 @@ namespace UE::NNERuntimeRDG::Private::Hlsl
 		}
 	};
 
-	bool ValidateElementWiseVariadicOperator(const NNECore::FAttributeMap& AttributeMap, TConstArrayView<ENNETensorDataType> InputTypes, TConstArrayView<NNECore::FSymbolicTensorShape> InputShapes)
+	bool ValidateElementWiseVariadicOperator(const NNE::FAttributeMap& AttributeMap, TConstArrayView<ENNETensorDataType> InputTypes, TConstArrayView<NNE::FSymbolicTensorShape> InputShapes)
 	{
 		bool bIsValid = true;
 
@@ -203,7 +203,7 @@ namespace UE::NNERuntimeRDG::Private::Hlsl
 		return bIsValid;
 	}
 
-	template<NNECore::Internal::EElementWiseVariadicOperatorType OpType>
+	template<NNE::Internal::EElementWiseVariadicOperatorType OpType>
 	FOperatorHlsl* CreateElementWiseVariadicOperator()
 	{
 		return new TElementWiseVariadic<OpType>();
@@ -211,7 +211,7 @@ namespace UE::NNERuntimeRDG::Private::Hlsl
 
 	bool RegisterElementWiseVariadicOperators(FOperatorRegistryHlsl& Registry)
 	{
-#define OP(Name) Registry.OpAdd(TEXT(#Name), CreateElementWiseVariadicOperator<NNECore::Internal::EElementWiseVariadicOperatorType::Name>)
+#define OP(Name) Registry.OpAdd(TEXT(#Name), CreateElementWiseVariadicOperator<NNE::Internal::EElementWiseVariadicOperatorType::Name>)
 		OP(Max);
 		OP(Min);
 		OP(Mean);
diff --git a/Engine/Plugins/Experimental/NNERuntimeRDG/Source/NNERuntimeRDG/Private/Hlsl/NNERuntimeRDGFlatten.cpp b/Engine/Plugins/Experimental/NNERuntimeRDG/Source/NNERuntimeRDG/Private/Hlsl/NNERuntimeRDGFlatten.cpp
--- a/Engine/Plugins/Experimental/NNERuntimeRDG/Source/NNERuntimeRDG/Private/Hlsl/NNERuntimeRDGFlatten.cpp
+++ b/Engine/Plugins/Experimental/NNERuntimeRDG/Source/NNERuntimeRDG/Private/Hlsl/NNERuntimeRDGFlatten.cpp
@@ -23,12 +23,12 @@ namespace UE::NNERuntimeRDG::Private::Hlsl
 
 	public:
 
-		virtual int PrepareOutputs(TConstArrayView<NNECore::Internal::FTensorRef> InputTensors, TArrayView<NNECore::Internal::FTensorRef> OutputTensors) const override
+		virtual int PrepareOutputs(TConstArrayView<NNE::Internal::FTensorRef> InputTensors, TArrayView<NNE::Internal::FTensorRef> OutputTensors) const override
 		{
 			check(InputTensors.Num() == 1);
 			check(OutputTensors.Num() == 1);
 
-			const NNECore::Internal::FTensor& X = *InputTensors[0];
+			const NNE::Internal::FTensor& X = *InputTensors[0];
 			TArray<uint32> OutputShapeData;
 			uint32 InnerDimSize = 1;
 
@@ -39,7 +39,7 @@ namespace UE::NNERuntimeRDG::Private::Hlsl
 			OutputShapeData.Add(InnerDimSize);
 			OutputShapeData.Add(X.GetShape().Volume() / InnerDimSize);
 
-			NNECore::FTensorShape OutputShape = NNECore::FTensorShape::Make(OutputShapeData);
+			NNE::FTensorShape OutputShape = NNE::FTensorShape::Make(OutputShapeData);
 
 			OutputTensors[0]->SetShape(OutputShape);
 			if (X.HasPreparedData())
@@ -50,7 +50,7 @@ namespace UE::NNERuntimeRDG::Private::Hlsl
 			return 0;
 		};
 
-		virtual bool Initialize(TConstArrayView<NNECore::FTensorDesc> InputTensorDescs, TConstArrayView<NNECore::FTensorDesc> OutputTensorDescs, const NNECore::FAttributeMap& Attributes) override
+		virtual bool Initialize(TConstArrayView<NNE::FTensorDesc> InputTensorDescs, TConstArrayView<NNE::FTensorDesc> OutputTensorDescs, const NNE::FAttributeMap& Attributes) override
 		{
 			check(InputTensorDescs.Num() == 1);
 			check(OutputTensorDescs.Num() == 1);
@@ -88,7 +88,7 @@ namespace UE::NNERuntimeRDG::Private::Hlsl
 		}
 	};
 
-	bool ValidateFlattenOperator(const NNECore::FAttributeMap& AttributeMap, TConstArrayView<ENNETensorDataType> InputTypes, TConstArrayView<NNECore::FSymbolicTensorShape> InputFlattens)
+	bool ValidateFlattenOperator(const NNE::FAttributeMap& AttributeMap, TConstArrayView<ENNETensorDataType> InputTypes, TConstArrayView<NNE::FSymbolicTensorShape> InputFlattens)
 	{
 		bool bIsValid = true;
 
diff --git a/Engine/Plugins/Experimental/NNERuntimeRDG/Source/NNERuntimeRDG/Private/Hlsl/NNERuntimeRDGGather.cpp b/Engine/Plugins/Experimental/NNERuntimeRDG/Source/NNERuntimeRDG/Private/Hlsl/NNERuntimeRDGGather.cpp
--- a/Engine/Plugins/Experimental/NNERuntimeRDG/Source/NNERuntimeRDG/Private/Hlsl/NNERuntimeRDGGather.cpp
+++ b/Engine/Plugins/Experimental/NNERuntimeRDG/Source/NNERuntimeRDG/Private/Hlsl/NNERuntimeRDGGather.cpp
@@ -29,15 +29,15 @@ namespace UE::NNERuntimeRDG::Private::Hlsl
 
 	public:
 
-		virtual int PrepareOutputs(TConstArrayView<NNECore::Internal::FTensorRef> InputTensors, TArrayView<NNECore::Internal::FTensorRef> OutputTensors) const override
+		virtual int PrepareOutputs(TConstArrayView<NNE::Internal::FTensorRef> InputTensors, TArrayView<NNE::Internal::FTensorRef> OutputTensors) const override
 		{
 			check(InputTensors.Num() == 2)
 			check(OutputTensors.Num() == 1)
 
-			const NNECore::Internal::FTensor& Data = *InputTensors[0];
-			const NNECore::Internal::FTensor& Indices = *InputTensors[1];
-			const NNECore::FTensorShape& DataShape = Data.GetShape();
-			const NNECore::FTensorShape& IndicesShape = Indices.GetShape();
+			const NNE::Internal::FTensor& Data = *InputTensors[0];
+			const NNE::Internal::FTensor& Indices = *InputTensors[1];
+			const NNE::FTensorShape& DataShape = Data.GetShape();
+			const NNE::FTensorShape& IndicesShape = Indices.GetShape();
 
 			const int32 OutputRank = IndicesShape.Rank() + DataShape.Rank() - 1;
 			TArray<uint32> OutputShapeData;
@@ -57,7 +57,7 @@ namespace UE::NNERuntimeRDG::Private::Hlsl
 				OutputShapeData.Add(DataShape.GetData()[DataRankIdx]);
 			}
 
-			NNECore::FTensorShape OutputShape = NNECore::FTensorShape::Make(OutputShapeData);
+			NNE::FTensorShape OutputShape = NNE::FTensorShape::Make(OutputShapeData);
 
 			check(OutputShape.Rank() == OutputRank);
 			OutputTensors[0]->SetShape(OutputShape);
@@ -73,16 +73,16 @@ namespace UE::NNERuntimeRDG::Private::Hlsl
 			return 0;
 		};
 		
-		virtual bool Initialize(TConstArrayView<NNECore::FTensorDesc> InputTensorDescs, TConstArrayView<NNECore::FTensorDesc> OutputTensorDescs, const NNECore::FAttributeMap& Attributes) override
+		virtual bool Initialize(TConstArrayView<NNE::FTensorDesc> InputTensorDescs, TConstArrayView<NNE::FTensorDesc> OutputTensorDescs, const NNE::FAttributeMap& Attributes) override
 		{
 			const int32 MaxNumDimensions = NNEHlslShaders::Internal::FGatherConstants::MAX_NUM_DIMENSIONS;
 			
 			check(InputTensorDescs.Num() == 2)
 			check(OutputTensorDescs.Num() == 1)
 
-			const NNECore::FTensorDesc& Data = InputTensorDescs[0];
-			const NNECore::FTensorDesc& Indices = InputTensorDescs[1];
-			const NNECore::FTensorDesc& Output = OutputTensorDescs[0];
+			const NNE::FTensorDesc& Data = InputTensorDescs[0];
+			const NNE::FTensorDesc& Indices = InputTensorDescs[1];
+			const NNE::FTensorDesc& Output = OutputTensorDescs[0];
 
 			if (Output.GetShape().Rank() > MaxNumDimensions)
 			{
@@ -159,7 +159,7 @@ namespace UE::NNERuntimeRDG::Private::Hlsl
 		}
 	};
 
-	bool ValidateGatherOperator(const NNECore::FAttributeMap& AttributeMap, TConstArrayView<ENNETensorDataType> InputTypes, TConstArrayView<NNECore::FSymbolicTensorShape> InputShapes)
+	bool ValidateGatherOperator(const NNE::FAttributeMap& AttributeMap, TConstArrayView<ENNETensorDataType> InputTypes, TConstArrayView<NNE::FSymbolicTensorShape> InputShapes)
 	{
 		bool bIsValid = true;
 
diff --git a/Engine/Plugins/Experimental/NNERuntimeRDG/Source/NNERuntimeRDG/Private/Hlsl/NNERuntimeRDGGemm.cpp b/Engine/Plugins/Experimental/NNERuntimeRDG/Source/NNERuntimeRDG/Private/Hlsl/NNERuntimeRDGGemm.cpp
--- a/Engine/Plugins/Experimental/NNERuntimeRDG/Source/NNERuntimeRDG/Private/Hlsl/NNERuntimeRDGGemm.cpp
+++ b/Engine/Plugins/Experimental/NNERuntimeRDG/Source/NNERuntimeRDG/Private/Hlsl/NNERuntimeRDGGemm.cpp
@@ -33,13 +33,13 @@ namespace UE::NNERuntimeRDG::Private::Hlsl
 
 	public:
 
-		virtual int PrepareOutputs(TConstArrayView<NNECore::Internal::FTensorRef> InputTensors, TArrayView<NNECore::Internal::FTensorRef> OutputTensors) const override
+		virtual int PrepareOutputs(TConstArrayView<NNE::Internal::FTensorRef> InputTensors, TArrayView<NNE::Internal::FTensorRef> OutputTensors) const override
 		{
 			check(InputTensors.Num() >= 2 && InputTensors.Num() <= 3);
 			check(OutputTensors.Num() == 1);
 
-			const NNECore::FTensorShape& InputA = InputTensors[0]->GetShape();
-			const NNECore::FTensorShape& InputB = InputTensors[1]->GetShape();
+			const NNE::FTensorShape& InputA = InputTensors[0]->GetShape();
+			const NNE::FTensorShape& InputB = InputTensors[1]->GetShape();
 			if (InputA.Rank() != 2 || InputB.Rank() != 2)
 			{
 				return -1;
@@ -52,19 +52,19 @@ namespace UE::NNERuntimeRDG::Private::Hlsl
 			OutputShapeData.Emplace(M);
 			OutputShapeData.Emplace(N);
 			
-			NNECore::FTensorShape OutputShape = NNECore::FTensorShape::Make(OutputShapeData);
+			NNE::FTensorShape OutputShape = NNE::FTensorShape::Make(OutputShapeData);
 			
 			OutputTensors[0]->SetShape(OutputShape);
 			return 0;
 		};
 		
-		virtual bool Initialize(TConstArrayView<NNECore::FTensorDesc> InputTensorDescs, TConstArrayView<NNECore::FTensorDesc> OutputTensorDescs, const NNECore::FAttributeMap& Attributes) override
+		virtual bool Initialize(TConstArrayView<NNE::FTensorDesc> InputTensorDescs, TConstArrayView<NNE::FTensorDesc> OutputTensorDescs, const NNE::FAttributeMap& Attributes) override
 		{
 			check(InputTensorDescs.Num() >= 2 && InputTensorDescs.Num() <= 3);
 			check(OutputTensorDescs.Num() == 1);
 
-            const NNECore::FTensorDesc& InputA = InputTensorDescs[0];
-			const NNECore::FTensorDesc& InputB = InputTensorDescs[1];
+            const NNE::FTensorDesc& InputA = InputTensorDescs[0];
+			const NNE::FTensorDesc& InputB = InputTensorDescs[1];
 
 			if (InputA.GetShape().Rank() != 2)
 			{
@@ -78,7 +78,7 @@ namespace UE::NNERuntimeRDG::Private::Hlsl
 			}
 			if (InputTensorDescs.Num() == 3)
 			{
-				const NNECore::FTensorDesc& InputC = InputTensorDescs[2];
+				const NNE::FTensorDesc& InputC = InputTensorDescs[2];
 				if (InputC.GetShape().Rank() > 2)
 				{
 					UE_LOG(LogNNE, Warning, TEXT("Gemm third input should be of rank 2 or less"));
@@ -159,7 +159,7 @@ namespace UE::NNERuntimeRDG::Private::Hlsl
 		}
 	};
 
-	bool ValidateGemmOperator(const NNECore::FAttributeMap& AttributeMap, TConstArrayView<ENNETensorDataType> InputTypes, TConstArrayView<NNECore::FSymbolicTensorShape> InputShapes)
+	bool ValidateGemmOperator(const NNE::FAttributeMap& AttributeMap, TConstArrayView<ENNETensorDataType> InputTypes, TConstArrayView<NNE::FSymbolicTensorShape> InputShapes)
 	{
 		bool bIsValid = true;
 
diff --git a/Engine/Plugins/Experimental/NNERuntimeRDG/Source/NNERuntimeRDG/Private/Hlsl/NNERuntimeRDGIdentity.cpp b/Engine/Plugins/Experimental/NNERuntimeRDG/Source/NNERuntimeRDG/Private/Hlsl/NNERuntimeRDGIdentity.cpp
--- a/Engine/Plugins/Experimental/NNERuntimeRDG/Source/NNERuntimeRDG/Private/Hlsl/NNERuntimeRDGIdentity.cpp
+++ b/Engine/Plugins/Experimental/NNERuntimeRDG/Source/NNERuntimeRDG/Private/Hlsl/NNERuntimeRDGIdentity.cpp
@@ -21,12 +21,12 @@ namespace UE::NNERuntimeRDG::Private::Hlsl
 
 	public:
 
-		virtual int PrepareOutputs(TConstArrayView<NNECore::Internal::FTensorRef> InputTensors, TArrayView<NNECore::Internal::FTensorRef> OutputTensors) const override
+		virtual int PrepareOutputs(TConstArrayView<NNE::Internal::FTensorRef> InputTensors, TArrayView<NNE::Internal::FTensorRef> OutputTensors) const override
 		{
 			check(InputTensors.Num() == 1);
 			check(OutputTensors.Num() == 1);
 
-			const NNECore::Internal::FTensor& X = *InputTensors[0];
+			const NNE::Internal::FTensor& X = *InputTensors[0];
 
 			OutputTensors[0]->SetShape(X.GetShape());
 			if (X.HasPreparedData())
@@ -37,7 +37,7 @@ namespace UE::NNERuntimeRDG::Private::Hlsl
 			return 0;
 		};
 
-		virtual bool Initialize(TConstArrayView<NNECore::FTensorDesc> InputTensorDescs, TConstArrayView<NNECore::FTensorDesc> OutputTensorDescs, const NNECore::FAttributeMap& Attributes) override
+		virtual bool Initialize(TConstArrayView<NNE::FTensorDesc> InputTensorDescs, TConstArrayView<NNE::FTensorDesc> OutputTensorDescs, const NNE::FAttributeMap& Attributes) override
 		{
 			check(InputTensorDescs.Num() == 1);
 			check(OutputTensorDescs.Num() == 1);
@@ -62,7 +62,7 @@ namespace UE::NNERuntimeRDG::Private::Hlsl
 		}
 	};
 
-	bool ValidateIdentityOperator(const NNECore::FAttributeMap& AttributeMap, TConstArrayView<ENNETensorDataType> InputTypes, TConstArrayView<NNECore::FSymbolicTensorShape> InputIdentitys)
+	bool ValidateIdentityOperator(const NNE::FAttributeMap& AttributeMap, TConstArrayView<ENNETensorDataType> InputTypes, TConstArrayView<NNE::FSymbolicTensorShape> InputIdentitys)
 	{
 		bool bIsValid = true;
 
diff --git a/Engine/Plugins/Experimental/NNERuntimeRDG/Source/NNERuntimeRDG/Private/Hlsl/NNERuntimeRDGInstanceNormalization.cpp b/Engine/Plugins/Experimental/NNERuntimeRDG/Source/NNERuntimeRDG/Private/Hlsl/NNERuntimeRDGInstanceNormalization.cpp
--- a/Engine/Plugins/Experimental/NNERuntimeRDG/Source/NNERuntimeRDG/Private/Hlsl/NNERuntimeRDGInstanceNormalization.cpp
+++ b/Engine/Plugins/Experimental/NNERuntimeRDG/Source/NNERuntimeRDG/Private/Hlsl/NNERuntimeRDGInstanceNormalization.cpp
@@ -9,7 +9,7 @@
 namespace
 {
 
-int ValidateInput(const UE::NNECore::FTensorShape& Input, const UE::NNECore::FTensorShape& Scale, const UE::NNECore::FTensorShape& Bias)
+int ValidateInput(const UE::NNE::FTensorShape& Input, const UE::NNE::FTensorShape& Scale, const UE::NNE::FTensorShape& Bias)
 {
 	if (Input.Rank() < 3)
 	{
@@ -68,7 +68,7 @@ namespace UE::NNERuntimeRDG::Private::Hlsl
 
 	public:
 
-		virtual int PrepareOutputs(TConstArrayView<NNECore::Internal::FTensorRef> InputTensors, TArrayView<NNECore::Internal::FTensorRef> OutputTensors) const override
+		virtual int PrepareOutputs(TConstArrayView<NNE::Internal::FTensorRef> InputTensors, TArrayView<NNE::Internal::FTensorRef> OutputTensors) const override
 		{
 			check(InputTensors.Num() == 3);
 			check(OutputTensors.Num() == 1);
@@ -83,7 +83,7 @@ namespace UE::NNERuntimeRDG::Private::Hlsl
 			return 0;
 		};
 
-		virtual bool Initialize(TConstArrayView<NNECore::FTensorDesc> InputTensorDescs, TConstArrayView<NNECore::FTensorDesc> OutputTensorDescs, const NNECore::FAttributeMap& Attributes) override
+		virtual bool Initialize(TConstArrayView<NNE::FTensorDesc> InputTensorDescs, TConstArrayView<NNE::FTensorDesc> OutputTensorDescs, const NNE::FAttributeMap& Attributes) override
 		{
 			using namespace UE::NNEHlslShaders::Internal;
 
@@ -146,7 +146,7 @@ namespace UE::NNERuntimeRDG::Private::Hlsl
 		}
 	};
 
-	bool ValidateInstanceNormalizationOperator(const NNECore::FAttributeMap& AttributeMap, TConstArrayView<ENNETensorDataType> InputTypes, TConstArrayView<NNECore::FSymbolicTensorShape> InputShapes)
+	bool ValidateInstanceNormalizationOperator(const NNE::FAttributeMap& AttributeMap, TConstArrayView<ENNETensorDataType> InputTypes, TConstArrayView<NNE::FSymbolicTensorShape> InputShapes)
 	{
 		bool bIsValid = true;
 
diff --git a/Engine/Plugins/Experimental/NNERuntimeRDG/Source/NNERuntimeRDG/Private/Hlsl/NNERuntimeRDGMatMul.cpp b/Engine/Plugins/Experimental/NNERuntimeRDG/Source/NNERuntimeRDG/Private/Hlsl/NNERuntimeRDGMatMul.cpp
--- a/Engine/Plugins/Experimental/NNERuntimeRDG/Source/NNERuntimeRDG/Private/Hlsl/NNERuntimeRDGMatMul.cpp
+++ b/Engine/Plugins/Experimental/NNERuntimeRDG/Source/NNERuntimeRDG/Private/Hlsl/NNERuntimeRDGMatMul.cpp
@@ -22,13 +22,13 @@ namespace UE::NNERuntimeRDG::Private::Hlsl
 
 	public:
 
-		virtual int PrepareOutputs(TConstArrayView<NNECore::Internal::FTensorRef> InputTensors, TArrayView<NNECore::Internal::FTensorRef> OutputTensors) const override
+		virtual int PrepareOutputs(TConstArrayView<NNE::Internal::FTensorRef> InputTensors, TArrayView<NNE::Internal::FTensorRef> OutputTensors) const override
 		{
 			check(InputTensors.Num() == 2);
 			check(OutputTensors.Num() == 1);
 
-			const NNECore::FTensorShape& InputA = InputTensors[0]->GetShape();
-			const NNECore::FTensorShape& InputB = InputTensors[1]->GetShape();
+			const NNE::FTensorShape& InputA = InputTensors[0]->GetShape();
+			const NNE::FTensorShape& InputB = InputTensors[1]->GetShape();
 
 			if (InputA.Rank() < 2)
 			{
@@ -65,20 +65,20 @@ namespace UE::NNERuntimeRDG::Private::Hlsl
 			OutputShapeData[OutputRank - 2] = InputA.GetData()[InputA.Rank() - 2];
 			OutputShapeData[OutputRank - 1] = InputB.GetData()[InputB.Rank() - 1];
 
-			NNECore::FTensorShape OutputShape = NNECore::FTensorShape::Make(OutputShapeData);
+			NNE::FTensorShape OutputShape = NNE::FTensorShape::Make(OutputShapeData);
 
 			OutputTensors[0]->SetShape(OutputShape);
 
 			return 0;
 		};
 
-		virtual bool Initialize(TConstArrayView<NNECore::FTensorDesc> InputTensorDescs, TConstArrayView<NNECore::FTensorDesc> OutputTensorDescs, const NNECore::FAttributeMap& Attributes) override
+		virtual bool Initialize(TConstArrayView<NNE::FTensorDesc> InputTensorDescs, TConstArrayView<NNE::FTensorDesc> OutputTensorDescs, const NNE::FAttributeMap& Attributes) override
 		{
 			check(InputTensorDescs.Num() == 2);
 			check(OutputTensorDescs.Num() == 1);
 
-			const NNECore::FTensorDesc& InputA = InputTensorDescs[0];
-			const NNECore::FTensorDesc& InputB = InputTensorDescs[1];
+			const NNE::FTensorDesc& InputA = InputTensorDescs[0];
+			const NNE::FTensorDesc& InputB = InputTensorDescs[1];
 
 			if (InputA.GetShape().Rank() < 2)
 			{
@@ -139,7 +139,7 @@ namespace UE::NNERuntimeRDG::Private::Hlsl
 		}
 	};
 
-	bool ValidateMatMulOperator(const NNECore::FAttributeMap& AttributeMap, TConstArrayView<ENNETensorDataType> InputTypes, TConstArrayView<NNECore::FSymbolicTensorShape> InputShapes)
+	bool ValidateMatMulOperator(const NNE::FAttributeMap& AttributeMap, TConstArrayView<ENNETensorDataType> InputTypes, TConstArrayView<NNE::FSymbolicTensorShape> InputShapes)
 	{
 		bool bIsValid = true;
 
diff --git a/Engine/Plugins/Experimental/NNERuntimeRDG/Source/NNERuntimeRDG/Private/Hlsl/NNERuntimeRDGPad.cpp b/Engine/Plugins/Experimental/NNERuntimeRDG/Source/NNERuntimeRDG/Private/Hlsl/NNERuntimeRDGPad.cpp
--- a/Engine/Plugins/Experimental/NNERuntimeRDG/Source/NNERuntimeRDG/Private/Hlsl/NNERuntimeRDGPad.cpp
+++ b/Engine/Plugins/Experimental/NNERuntimeRDG/Source/NNERuntimeRDG/Private/Hlsl/NNERuntimeRDGPad.cpp
@@ -28,12 +28,12 @@ namespace UE::NNERuntimeRDG::Private::Hlsl
 
 	public:
 
-		virtual int PrepareOutputs(TConstArrayView<NNECore::Internal::FTensorRef> InputTensors, TArrayView<NNECore::Internal::FTensorRef> OutputTensors) const override
+		virtual int PrepareOutputs(TConstArrayView<NNE::Internal::FTensorRef> InputTensors, TArrayView<NNE::Internal::FTensorRef> OutputTensors) const override
 		{
 			check(InputTensors.Num() == 1);
 			check(OutputTensors.Num() == 1);
 
-			const NNECore::Internal::FTensor& X = *InputTensors[0];
+			const NNE::Internal::FTensor& X = *InputTensors[0];
 
 			TArray<uint32> OutputShapeData;
 			for (int32 i = 0; i < X.GetShape().Rank(); ++i)
@@ -49,13 +49,13 @@ namespace UE::NNERuntimeRDG::Private::Hlsl
 				OutputShapeData.Emplace(OutputDim);
 			}
 
-			NNECore::FTensorShape OutputShape = NNECore::FTensorShape::Make(OutputShapeData);
+			NNE::FTensorShape OutputShape = NNE::FTensorShape::Make(OutputShapeData);
 			OutputTensors[0]->SetShape(OutputShape);
 
 			return 0;
 		};
 
-		virtual bool Initialize(TConstArrayView<NNECore::FTensorDesc> InputTensorDescs, TConstArrayView<NNECore::FTensorDesc> OutputTensorDescs, const NNECore::FAttributeMap& Attributes) override
+		virtual bool Initialize(TConstArrayView<NNE::FTensorDesc> InputTensorDescs, TConstArrayView<NNE::FTensorDesc> OutputTensorDescs, const NNE::FAttributeMap& Attributes) override
 		{
 			using namespace UE::NNEHlslShaders::Internal;
 			
@@ -125,7 +125,7 @@ namespace UE::NNERuntimeRDG::Private::Hlsl
 		}
 	};
 
-	bool ValidatePadOperator(const NNECore::FAttributeMap& AttributeMap, TConstArrayView<ENNETensorDataType> InputTypes, TConstArrayView<NNECore::FSymbolicTensorShape> InputShapes)
+	bool ValidatePadOperator(const NNE::FAttributeMap& AttributeMap, TConstArrayView<ENNETensorDataType> InputTypes, TConstArrayView<NNE::FSymbolicTensorShape> InputShapes)
 	{
 		bool bIsValid = true;
 
diff --git a/Engine/Plugins/Experimental/NNERuntimeRDG/Source/NNERuntimeRDG/Private/Hlsl/NNERuntimeRDGReshape.cpp b/Engine/Plugins/Experimental/NNERuntimeRDG/Source/NNERuntimeRDG/Private/Hlsl/NNERuntimeRDGReshape.cpp
--- a/Engine/Plugins/Experimental/NNERuntimeRDG/Source/NNERuntimeRDG/Private/Hlsl/NNERuntimeRDGReshape.cpp
+++ b/Engine/Plugins/Experimental/NNERuntimeRDG/Source/NNERuntimeRDG/Private/Hlsl/NNERuntimeRDGReshape.cpp
@@ -24,13 +24,13 @@ namespace UE::NNERuntimeRDG::Private::Hlsl
 
 	public:
 
-		virtual int PrepareOutputs(TConstArrayView<NNECore::Internal::FTensorRef> InputTensors, TArrayView<NNECore::Internal::FTensorRef> OutputTensors) const override
+		virtual int PrepareOutputs(TConstArrayView<NNE::Internal::FTensorRef> InputTensors, TArrayView<NNE::Internal::FTensorRef> OutputTensors) const override
 		{
 			check(InputTensors.Num() == 2);
 			check(OutputTensors.Num() == 1);
 
-			const NNECore::Internal::FTensor& X = *InputTensors[0];
-			const NNECore::Internal::FTensor& ShapeTensor = *InputTensors[1];
+			const NNE::Internal::FTensor& X = *InputTensors[0];
+			const NNE::Internal::FTensor& ShapeTensor = *InputTensors[1];
 			TArray<uint32> OutputShapeData;
 
 			check(ShapeTensor.GetDataType() == ENNETensorDataType::Int64);
@@ -46,7 +46,7 @@ namespace UE::NNERuntimeRDG::Private::Hlsl
 				return -1;
 			}
 
-			NNECore::FTensorShape OutputShape = NNECore::FTensorShape::Make(OutputShapeData);
+			NNE::FTensorShape OutputShape = NNE::FTensorShape::Make(OutputShapeData);
 			
 			OutputTensors[0]->SetShape(OutputShape);
 			if (X.HasPreparedData())
@@ -57,7 +57,7 @@ namespace UE::NNERuntimeRDG::Private::Hlsl
 			return 0;
 		};
 
-		virtual bool Initialize(TConstArrayView<NNECore::FTensorDesc> InputTensorDescs, TConstArrayView<NNECore::FTensorDesc> OutputTensorDescs, const NNECore::FAttributeMap& Attributes) override
+		virtual bool Initialize(TConstArrayView<NNE::FTensorDesc> InputTensorDescs, TConstArrayView<NNE::FTensorDesc> OutputTensorDescs, const NNE::FAttributeMap& Attributes) override
 		{
 			check(InputTensorDescs.Num() == 2);
 			check(OutputTensorDescs.Num() == 1);
@@ -84,7 +84,7 @@ namespace UE::NNERuntimeRDG::Private::Hlsl
 		}
 	};
 
-	bool ValidateReshapeOperator(const NNECore::FAttributeMap& AttributeMap, TConstArrayView<ENNETensorDataType> InputTypes, TConstArrayView<NNECore::FSymbolicTensorShape> InputReshapes)
+	bool ValidateReshapeOperator(const NNE::FAttributeMap& AttributeMap, TConstArrayView<ENNETensorDataType> InputTypes, TConstArrayView<NNE::FSymbolicTensorShape> InputReshapes)
 	{
 		bool bIsValid = true;
 
diff --git a/Engine/Plugins/Experimental/NNERuntimeRDG/Source/NNERuntimeRDG/Private/Hlsl/NNERuntimeRDGShape.cpp b/Engine/Plugins/Experimental/NNERuntimeRDG/Source/NNERuntimeRDG/Private/Hlsl/NNERuntimeRDGShape.cpp
--- a/Engine/Plugins/Experimental/NNERuntimeRDG/Source/NNERuntimeRDG/Private/Hlsl/NNERuntimeRDGShape.cpp
+++ b/Engine/Plugins/Experimental/NNERuntimeRDG/Source/NNERuntimeRDG/Private/Hlsl/NNERuntimeRDGShape.cpp
@@ -18,19 +18,19 @@ namespace UE::NNERuntimeRDG::Private::Hlsl
 
 	public:
 
-		virtual int PrepareOutputs(TConstArrayView<NNECore::Internal::FTensorRef> InputTensors, TArrayView<NNECore::Internal::FTensorRef> OutputTensors) const override
+		virtual int PrepareOutputs(TConstArrayView<NNE::Internal::FTensorRef> InputTensors, TArrayView<NNE::Internal::FTensorRef> OutputTensors) const override
 		{
 			check(InputTensors.Num() == 1);
 			check(OutputTensors.Num() == 1);
 			check(OutputTensors[0]->GetDataType() == ENNETensorDataType::Int64);
 
-			const NNECore::Internal::FTensor& X = *InputTensors[0];
+			const NNE::Internal::FTensor& X = *InputTensors[0];
 
 			TArray<uint32> OutputShapeData;
 			
 			OutputShapeData.Emplace(X.GetShape().Rank());
 			
-			NNECore::FTensorShape OutputShape = NNECore::FTensorShape::Make(OutputShapeData);
+			NNE::FTensorShape OutputShape = NNE::FTensorShape::Make(OutputShapeData);
 			TArray<int64> OutputData;
 			
 			OutputTensors[0]->SetShape(OutputShape);
@@ -43,7 +43,7 @@ namespace UE::NNERuntimeRDG::Private::Hlsl
 			return 0;
 		};
 
-		virtual bool Initialize(TConstArrayView<NNECore::FTensorDesc> InputTensorDescs, TConstArrayView<NNECore::FTensorDesc> OutputTensorDescs, const NNECore::FAttributeMap& Attributes) override
+		virtual bool Initialize(TConstArrayView<NNE::FTensorDesc> InputTensorDescs, TConstArrayView<NNE::FTensorDesc> OutputTensorDescs, const NNE::FAttributeMap& Attributes) override
 		{
 			check(InputTensorDescs.Num() == 1);
 			check(OutputTensorDescs.Num() == 1);
@@ -63,7 +63,7 @@ namespace UE::NNERuntimeRDG::Private::Hlsl
 		}
 	};
 
-	bool ValidateShapeOperator(const NNECore::FAttributeMap& AttributeMap, TConstArrayView<ENNETensorDataType> InputTypes, TConstArrayView<NNECore::FSymbolicTensorShape> InputShapes)
+	bool ValidateShapeOperator(const NNE::FAttributeMap& AttributeMap, TConstArrayView<ENNETensorDataType> InputTypes, TConstArrayView<NNE::FSymbolicTensorShape> InputShapes)
 	{
 		bool bIsValid = true;
 
diff --git a/Engine/Plugins/Experimental/NNERuntimeRDG/Source/NNERuntimeRDG/Private/Hlsl/NNERuntimeRDGSize.cpp b/Engine/Plugins/Experimental/NNERuntimeRDG/Source/NNERuntimeRDG/Private/Hlsl/NNERuntimeRDGSize.cpp
--- a/Engine/Plugins/Experimental/NNERuntimeRDG/Source/NNERuntimeRDG/Private/Hlsl/NNERuntimeRDGSize.cpp
+++ b/Engine/Plugins/Experimental/NNERuntimeRDG/Source/NNERuntimeRDG/Private/Hlsl/NNERuntimeRDGSize.cpp
@@ -20,13 +20,13 @@ namespace UE::NNERuntimeRDG::Private::Hlsl
 
 	public:
 
-		virtual int PrepareOutputs(TConstArrayView<NNECore::Internal::FTensorRef> InputTensors, TArrayView<NNECore::Internal::FTensorRef> OutputTensors) const override
+		virtual int PrepareOutputs(TConstArrayView<NNE::Internal::FTensorRef> InputTensors, TArrayView<NNE::Internal::FTensorRef> OutputTensors) const override
 		{
 			check(InputTensors.Num() == 1);
 			check(OutputTensors.Num() == 1);
 
-			const NNECore::Internal::FTensor& X = *InputTensors[0];
-			NNECore::FTensorShape ScalarShape = NNECore::FTensorShape::Make({});
+			const NNE::Internal::FTensor& X = *InputTensors[0];
+			NNE::FTensorShape ScalarShape = NNE::FTensorShape::Make({});
 			TArray<int64> OutputData;
 			
 			OutputTensors[0]->SetShape(ScalarShape);
@@ -36,7 +36,7 @@ namespace UE::NNERuntimeRDG::Private::Hlsl
 			return 0;
 		};
 
-		virtual bool Initialize(TConstArrayView<NNECore::FTensorDesc> InputTensorDescs, TConstArrayView<NNECore::FTensorDesc> OutputTensorDescs, const NNECore::FAttributeMap& Attributes) override
+		virtual bool Initialize(TConstArrayView<NNE::FTensorDesc> InputTensorDescs, TConstArrayView<NNE::FTensorDesc> OutputTensorDescs, const NNE::FAttributeMap& Attributes) override
 		{
 			check(InputTensorDescs.Num() == 1);
 			check(OutputTensorDescs.Num() == 1);
@@ -51,7 +51,7 @@ namespace UE::NNERuntimeRDG::Private::Hlsl
 		}
 	};
 
-	bool ValidateSizeOperator(const NNECore::FAttributeMap& AttributeMap, TConstArrayView<ENNETensorDataType> InputTypes, TConstArrayView<NNECore::FSymbolicTensorShape> InputSizes)
+	bool ValidateSizeOperator(const NNE::FAttributeMap& AttributeMap, TConstArrayView<ENNETensorDataType> InputTypes, TConstArrayView<NNE::FSymbolicTensorShape> InputSizes)
 	{
 		bool bIsValid = true;
 
diff --git a/Engine/Plugins/Experimental/NNERuntimeRDG/Source/NNERuntimeRDG/Private/Hlsl/NNERuntimeRDGSlice.cpp b/Engine/Plugins/Experimental/NNERuntimeRDG/Source/NNERuntimeRDG/Private/Hlsl/NNERuntimeRDGSlice.cpp
--- a/Engine/Plugins/Experimental/NNERuntimeRDG/Source/NNERuntimeRDG/Private/Hlsl/NNERuntimeRDGSlice.cpp
+++ b/Engine/Plugins/Experimental/NNERuntimeRDG/Source/NNERuntimeRDG/Private/Hlsl/NNERuntimeRDGSlice.cpp
@@ -23,7 +23,7 @@ namespace UE::NNERuntimeRDG::Private::Hlsl
 
 	public:
 
-		virtual int PrepareOutputs(TConstArrayView<NNECore::Internal::FTensorRef> InputTensors, TArrayView<NNECore::Internal::FTensorRef> OutputTensors) const override
+		virtual int PrepareOutputs(TConstArrayView<NNE::Internal::FTensorRef> InputTensors, TArrayView<NNE::Internal::FTensorRef> OutputTensors) const override
 		{
 			check(InputTensors.Num() == 1);
 			check(OutputTensors.Num() == 1);
@@ -80,7 +80,7 @@ namespace UE::NNERuntimeRDG::Private::Hlsl
 				OutputShapeData.Add(End[i] - Start[i]);
 			}
 			
-			NNECore::FTensorShape OutputShape = NNECore::FTensorShape::Make(OutputShapeData);
+			NNE::FTensorShape OutputShape = NNE::FTensorShape::Make(OutputShapeData);
 
 			OutputTensors[0]->SetShape(OutputShape);
 			
@@ -95,7 +95,7 @@ namespace UE::NNERuntimeRDG::Private::Hlsl
 			return 0;
 		};
 
-		virtual bool Initialize(TConstArrayView<NNECore::FTensorDesc> InputTensorDescs, TConstArrayView<NNECore::FTensorDesc> OutputTensorDescs, const NNECore::FAttributeMap& Attributes) override
+		virtual bool Initialize(TConstArrayView<NNE::FTensorDesc> InputTensorDescs, TConstArrayView<NNE::FTensorDesc> OutputTensorDescs, const NNE::FAttributeMap& Attributes) override
 		{
 			check(InputTensorDescs.Num() == 1);
 			check(OutputTensorDescs.Num() == 1);
@@ -127,7 +127,7 @@ namespace UE::NNERuntimeRDG::Private::Hlsl
 		}
 	};
 
-	bool ValidateSliceOperator(const NNECore::FAttributeMap& AttributeMap, TConstArrayView<ENNETensorDataType> InputTypes, TConstArrayView<NNECore::FSymbolicTensorShape> InputSlices)
+	bool ValidateSliceOperator(const NNE::FAttributeMap& AttributeMap, TConstArrayView<ENNETensorDataType> InputTypes, TConstArrayView<NNE::FSymbolicTensorShape> InputSlices)
 	{
 		bool bIsValid = true;
 
diff --git a/Engine/Plugins/Experimental/NNERuntimeRDG/Source/NNERuntimeRDG/Private/Hlsl/NNERuntimeRDGSqueeze.cpp b/Engine/Plugins/Experimental/NNERuntimeRDG/Source/NNERuntimeRDG/Private/Hlsl/NNERuntimeRDGSqueeze.cpp
--- a/Engine/Plugins/Experimental/NNERuntimeRDG/Source/NNERuntimeRDG/Private/Hlsl/NNERuntimeRDGSqueeze.cpp
+++ b/Engine/Plugins/Experimental/NNERuntimeRDG/Source/NNERuntimeRDG/Private/Hlsl/NNERuntimeRDGSqueeze.cpp
@@ -23,12 +23,12 @@ namespace UE::NNERuntimeRDG::Private::Hlsl
 
 	public:
 
-		virtual int PrepareOutputs(TConstArrayView<NNECore::Internal::FTensorRef> InputTensors, TArrayView<NNECore::Internal::FTensorRef> OutputTensors) const override
+		virtual int PrepareOutputs(TConstArrayView<NNE::Internal::FTensorRef> InputTensors, TArrayView<NNE::Internal::FTensorRef> OutputTensors) const override
 		{
 			check(InputTensors.Num() == 1);
 			check(OutputTensors.Num() == 1);
 
-			const NNECore::Internal::FTensor& X = *InputTensors[0];
+			const NNE::Internal::FTensor& X = *InputTensors[0];
 			TArray<uint32> OutputShapeData(X.GetShape().GetData());
 
 			for (int32 Idx = Axes.Num() - 1; Idx >= 0; --Idx)
@@ -44,7 +44,7 @@ namespace UE::NNERuntimeRDG::Private::Hlsl
 				OutputShapeData.RemoveAt(Axe);
 			}
 			
-			NNECore::FTensorShape OutputShape = NNECore::FTensorShape::Make(OutputShapeData);
+			NNE::FTensorShape OutputShape = NNE::FTensorShape::Make(OutputShapeData);
 			
 			OutputTensors[0]->SetShape(OutputShape);
 			if (X.HasPreparedData())
@@ -55,7 +55,7 @@ namespace UE::NNERuntimeRDG::Private::Hlsl
 			return 0;
 		};
 
-		virtual bool Initialize(TConstArrayView<NNECore::FTensorDesc> InputTensorDescs, TConstArrayView<NNECore::FTensorDesc> OutputTensorDescs, const NNECore::FAttributeMap& Attributes) override
+		virtual bool Initialize(TConstArrayView<NNE::FTensorDesc> InputTensorDescs, TConstArrayView<NNE::FTensorDesc> OutputTensorDescs, const NNE::FAttributeMap& Attributes) override
 		{
 			check(InputTensorDescs.Num() == 1);
 			check(OutputTensorDescs.Num() == 1);
@@ -97,7 +97,7 @@ namespace UE::NNERuntimeRDG::Private::Hlsl
 		}
 	};
 
-	bool ValidateSqueezeOperator(const NNECore::FAttributeMap& AttributeMap, TConstArrayView<ENNETensorDataType> InputTypes, TConstArrayView<NNECore::FSymbolicTensorShape> InputSqueezes)
+	bool ValidateSqueezeOperator(const NNE::FAttributeMap& AttributeMap, TConstArrayView<ENNETensorDataType> InputTypes, TConstArrayView<NNE::FSymbolicTensorShape> InputSqueezes)
 	{
 		bool bIsValid = true;
 
diff --git a/Engine/Plugins/Experimental/NNERuntimeRDG/Source/NNERuntimeRDG/Private/Hlsl/NNERuntimeRDGUnsqueeze.cpp b/Engine/Plugins/Experimental/NNERuntimeRDG/Source/NNERuntimeRDG/Private/Hlsl/NNERuntimeRDGUnsqueeze.cpp
--- a/Engine/Plugins/Experimental/NNERuntimeRDG/Source/NNERuntimeRDG/Private/Hlsl/NNERuntimeRDGUnsqueeze.cpp
+++ b/Engine/Plugins/Experimental/NNERuntimeRDG/Source/NNERuntimeRDG/Private/Hlsl/NNERuntimeRDGUnsqueeze.cpp
@@ -23,12 +23,12 @@ namespace UE::NNERuntimeRDG::Private::Hlsl
 
 	public:
 
-		virtual int PrepareOutputs(TConstArrayView<NNECore::Internal::FTensorRef> InputTensors, TArrayView<NNECore::Internal::FTensorRef> OutputTensors) const override
+		virtual int PrepareOutputs(TConstArrayView<NNE::Internal::FTensorRef> InputTensors, TArrayView<NNE::Internal::FTensorRef> OutputTensors) const override
 		{
 			check(InputTensors.Num() == 1);
 			check(OutputTensors.Num() == 1);
 
-			const NNECore::Internal::FTensor& X = *InputTensors[0];
+			const NNE::Internal::FTensor& X = *InputTensors[0];
 			TArray<uint32> OutputShapeData(X.GetShape().GetData());
 
 			for (int32 Axe : Axes)
@@ -36,7 +36,7 @@ namespace UE::NNERuntimeRDG::Private::Hlsl
 				OutputShapeData.Insert(1, Axe);
 			}
 			
-			NNECore::FTensorShape OutputShape = NNECore::FTensorShape::Make(OutputShapeData);
+			NNE::FTensorShape OutputShape = NNE::FTensorShape::Make(OutputShapeData);
 			
 			OutputTensors[0]->SetShape(OutputShape);
 			if (X.HasPreparedData())
@@ -47,7 +47,7 @@ namespace UE::NNERuntimeRDG::Private::Hlsl
 			return 0;
 		};
 
-		virtual bool Initialize(TConstArrayView<NNECore::FTensorDesc> InputTensorDescs, TConstArrayView<NNECore::FTensorDesc> OutputTensorDescs, const NNECore::FAttributeMap& Attributes) override
+		virtual bool Initialize(TConstArrayView<NNE::FTensorDesc> InputTensorDescs, TConstArrayView<NNE::FTensorDesc> OutputTensorDescs, const NNE::FAttributeMap& Attributes) override
 		{
 			check(InputTensorDescs.Num() == 1);
 			check(OutputTensorDescs.Num() == 1);
@@ -89,7 +89,7 @@ namespace UE::NNERuntimeRDG::Private::Hlsl
 		}
 	};
 
-	bool ValidateUnsqueezeOperator(const NNECore::FAttributeMap& AttributeMap, TConstArrayView<ENNETensorDataType> InputTypes, TConstArrayView<NNECore::FSymbolicTensorShape> InputUnsqueezes)
+	bool ValidateUnsqueezeOperator(const NNE::FAttributeMap& AttributeMap, TConstArrayView<ENNETensorDataType> InputTypes, TConstArrayView<NNE::FSymbolicTensorShape> InputUnsqueezes)
 	{
 		bool bIsValid = true;
 
diff --git a/Engine/Plugins/Experimental/NNERuntimeRDG/Source/NNERuntimeRDG/Private/Hlsl/NNERuntimeRDGUpsample.cpp b/Engine/Plugins/Experimental/NNERuntimeRDG/Source/NNERuntimeRDG/Private/Hlsl/NNERuntimeRDGUpsample.cpp
--- a/Engine/Plugins/Experimental/NNERuntimeRDG/Source/NNERuntimeRDG/Private/Hlsl/NNERuntimeRDGUpsample.cpp
+++ b/Engine/Plugins/Experimental/NNERuntimeRDG/Source/NNERuntimeRDG/Private/Hlsl/NNERuntimeRDGUpsample.cpp
@@ -22,13 +22,13 @@ namespace UE::NNERuntimeRDG::Private::Hlsl
 
 	public:
 
-		virtual int PrepareOutputs(TConstArrayView<NNECore::Internal::FTensorRef> InputTensors, TArrayView<NNECore::Internal::FTensorRef> OutputTensors) const override
+		virtual int PrepareOutputs(TConstArrayView<NNE::Internal::FTensorRef> InputTensors, TArrayView<NNE::Internal::FTensorRef> OutputTensors) const override
 		{
 			check(InputTensors.Num() == 2);
 			check(OutputTensors.Num() == 1);
 
-			const NNECore::Internal::FTensor& X = *InputTensors[0];
-			const NNECore::Internal::FTensor& Scales = *InputTensors[1];
+			const NNE::Internal::FTensor& X = *InputTensors[0];
+			const NNE::Internal::FTensor& Scales = *InputTensors[1];
 			
 			if (!Scales.HasPreparedData())
 			{
@@ -51,13 +51,13 @@ namespace UE::NNERuntimeRDG::Private::Hlsl
 				OutputShapeData.Emplace(FMath::FloorToInt32(X.GetShape().GetData()[i] * ScalesData[i]));
 			}
 
-			NNECore::FTensorShape OutputShape = NNECore::FTensorShape::Make(OutputShapeData);
+			NNE::FTensorShape OutputShape = NNE::FTensorShape::Make(OutputShapeData);
 			OutputTensors[0]->SetShape(OutputShape);
 
 			return 0;
 		};
 
-		virtual bool Initialize(TConstArrayView<NNECore::FTensorDesc> InputTensorDescs, TConstArrayView<NNECore::FTensorDesc> OutputTensorDescs, const NNECore::FAttributeMap& Attributes) override
+		virtual bool Initialize(TConstArrayView<NNE::FTensorDesc> InputTensorDescs, TConstArrayView<NNE::FTensorDesc> OutputTensorDescs, const NNE::FAttributeMap& Attributes) override
 		{
 			check(InputTensorDescs.Num() == 2);
 			check(OutputTensorDescs.Num() == 1);
@@ -116,7 +116,7 @@ namespace UE::NNERuntimeRDG::Private::Hlsl
 		}
 	};
 
-	bool ValidateUpsampleOperator(const NNECore::FAttributeMap& AttributeMap, TConstArrayView<ENNETensorDataType> InputTypes, TConstArrayView<NNECore::FSymbolicTensorShape> InputShapes)
+	bool ValidateUpsampleOperator(const NNE::FAttributeMap& AttributeMap, TConstArrayView<ENNETensorDataType> InputTypes, TConstArrayView<NNE::FSymbolicTensorShape> InputShapes)
 	{
 		bool bIsValid = true;
 
diff --git a/Engine/Plugins/Experimental/NNERuntimeRDG/Source/NNERuntimeRDG/Private/NNERuntimeRDGBase.cpp b/Engine/Plugins/Experimental/NNERuntimeRDG/Source/NNERuntimeRDG/Private/NNERuntimeRDGBase.cpp
--- a/Engine/Plugins/Experimental/NNERuntimeRDG/Source/NNERuntimeRDG/Private/NNERuntimeRDGBase.cpp
+++ b/Engine/Plugins/Experimental/NNERuntimeRDG/Source/NNERuntimeRDG/Private/NNERuntimeRDGBase.cpp
@@ -17,9 +17,9 @@ namespace UE::NNERuntimeRDG::Private
 //
 //
 bool AlwaysValidValidationFunction(
-	const NNECore::FAttributeMap& AttributeMap, 
+	const NNE::FAttributeMap& AttributeMap, 
 	TConstArrayView<ENNETensorDataType> InputTensorTypes,
-	TConstArrayView<NNECore::FSymbolicTensorShape> InputShapes)
+	TConstArrayView<NNE::FSymbolicTensorShape> InputShapes)
 {
 	return true;
 }
@@ -106,7 +106,7 @@ void FAttributeValidator::AddRequired(const FString& Name, ENNEAttributeDataType
 	RequiredAttributes.Emplace(Name, Type);
 }
 
-bool FAttributeValidator::Validate(const NNECore::FAttributeMap& AttributesToValidate)
+bool FAttributeValidator::Validate(const NNE::FAttributeMap& AttributesToValidate)
 {
 	bool bAreAttributesValid = true;
 
diff --git a/Engine/Plugins/Experimental/NNERuntimeRDG/Source/NNERuntimeRDG/Private/NNERuntimeRDGBase.h b/Engine/Plugins/Experimental/NNERuntimeRDG/Source/NNERuntimeRDG/Private/NNERuntimeRDGBase.h
--- a/Engine/Plugins/Experimental/NNERuntimeRDG/Source/NNERuntimeRDG/Private/NNERuntimeRDGBase.h
+++ b/Engine/Plugins/Experimental/NNERuntimeRDG/Source/NNERuntimeRDG/Private/NNERuntimeRDGBase.h
@@ -22,7 +22,7 @@ END_SHADER_PARAMETER_STRUCT()
 class FRDGBuilder;
 struct FNNEModelRaw;
 
-namespace UE::NNECore { class FTensorDesc; }
+namespace UE::NNE { class FTensorDesc; }
 
 namespace UE::NNERuntimeRDG::Private
 {
@@ -33,7 +33,7 @@ namespace UE::NNERuntimeRDG::Private
 struct IPrepareOperator
 {
 	virtual ~IPrepareOperator() = default;
-	virtual int PrepareOutputs(TConstArrayView<NNECore::Internal::FTensorRef> InputTensors, TArrayView<NNECore::Internal::FTensorRef> OutputTensors) const = 0;
+	virtual int PrepareOutputs(TConstArrayView<NNE::Internal::FTensorRef> InputTensors, TArrayView<NNE::Internal::FTensorRef> OutputTensors) const = 0;
 };
 
 /**
@@ -44,12 +44,12 @@ struct IOperatorRDG
 	virtual ~IOperatorRDG() = default;
 };
 
-class FTensorRDG : public NNECore::Internal::FTensor
+class FTensorRDG : public NNE::Internal::FTensor
 {
 	FRDGBufferRef Buffer{};
 	
 public:
-	static FTensorRDG Make(const NNECore::FTensorDesc& TensorDesc, const NNECore::FTensorShape& Shape, FRDGBufferRef Buffer)
+	static FTensorRDG Make(const NNE::FTensorDesc& TensorDesc, const NNE::FTensorShape& Shape, FRDGBufferRef Buffer)
 	{
 		check(Shape.IsCompatibleWith(TensorDesc.GetShape()));
 		FTensorRDG TensorRDG;
@@ -58,7 +58,7 @@ public:
 		TensorRDG.DataType = TensorDesc.GetDataType();
 		TensorRDG.Shape = Shape;
 		TensorRDG.Volume = Shape.Volume();
-		TensorRDG.DataSize = (uint64)NNECore::GetTensorDataTypeSizeInBytes(TensorRDG.DataType) * TensorRDG.Volume;
+		TensorRDG.DataSize = (uint64)NNE::GetTensorDataTypeSizeInBytes(TensorRDG.DataType) * TensorRDG.Volume;
 		TensorRDG.Buffer = Buffer;
 		check(TensorRDG.Volume <= TNumericLimits<uint32>::Max());
 		return TensorRDG;
@@ -74,7 +74,7 @@ using FTensorRDGArray = TArray<FTensorRDG, TInlineAllocator<16>>;
 using FTensorRDGRefArray = TArray<FTensorRDGRef, TInlineAllocator<64>>;
 using FIntArray = TArray<int32, TInlineAllocator<16>>;
 
-bool AlwaysValidValidationFunction(const NNECore::FAttributeMap& AttributeMap, TConstArrayView<ENNETensorDataType> InputTypes, TConstArrayView<NNECore::FSymbolicTensorShape> InputShapes);
+bool AlwaysValidValidationFunction(const NNE::FAttributeMap& AttributeMap, TConstArrayView<ENNETensorDataType> InputTypes, TConstArrayView<NNE::FSymbolicTensorShape> InputShapes);
 
 class FInputValidator
 {
@@ -98,7 +98,7 @@ class FAttributeValidator
 public:
 	void AddOptional(const FString& Name, ENNEAttributeDataType Type);
 	void AddRequired(const FString& Name, ENNEAttributeDataType Type);
-	bool Validate(const NNECore::FAttributeMap& AttributesToValidate);
+	bool Validate(const NNE::FAttributeMap& AttributesToValidate);
 
 private:
 	struct FEntry
@@ -129,7 +129,7 @@ class TOperatorRegistryRDG
 public:
 
 	typedef TOperatorType* (*OperatorCreateFunc)();
-	typedef bool (*OperatorValidateFunc)(const NNECore::FAttributeMap& AttributeMap, TConstArrayView<ENNETensorDataType> InputTypes, TConstArrayView<NNECore::FSymbolicTensorShape> InputShapes);
+	typedef bool (*OperatorValidateFunc)(const NNE::FAttributeMap& AttributeMap, TConstArrayView<ENNETensorDataType> InputTypes, TConstArrayView<NNE::FSymbolicTensorShape> InputShapes);
 
 	static TOperatorRegistryRDG* Get()
 	{
@@ -187,7 +187,7 @@ private:
  * Validator for RDG ML operators
  */
 template<class TOperatorType>
-class TModelValidatorRDG : public UE::NNECore::Internal::IModelValidator
+class TModelValidatorRDG : public UE::NNE::Internal::IModelValidator
 {
 public:
 	virtual FString GetName() const 
@@ -195,7 +195,7 @@ public:
 		return TEXT("RDG Model validator");
 	}
 	
-	virtual bool ValidateModel(const FNNEModelRaw& InputModel, const UE::NNECore::Internal::FOptimizerOptionsMap& Options) const override
+	virtual bool ValidateModel(const FNNEModelRaw& InputModel, const UE::NNE::Internal::FOptimizerOptionsMap& Options) const override
 	{
 		FNNERuntimeFormat	Format;
 
@@ -215,13 +215,13 @@ public:
 		for (int32 Idx = 0; Idx < Format.Operators.Num(); ++Idx)
 		{
 			TArray<ENNETensorDataType> InputTensorTypes;
-			TArray<NNECore::FSymbolicTensorShape> InputTensorShapes;
-			NNECore::FAttributeMap AttributeMap;
+			TArray<NNE::FSymbolicTensorShape> InputTensorShapes;
+			NNE::FAttributeMap AttributeMap;
 			
 			for (int32 InputTensorIndex: Format.Operators[Idx].InTensors)
 			{
 				InputTensorTypes.Add(Format.Tensors[InputTensorIndex].DataType);
-				InputTensorShapes.Add(NNECore::FSymbolicTensorShape::Make(Format.Tensors[InputTensorIndex].Shape));
+				InputTensorShapes.Add(NNE::FSymbolicTensorShape::Make(Format.Tensors[InputTensorIndex].Shape));
 			}
 			for (const FNNEFormatAttributeDesc& Desc : Format.Operators[Idx].Attributes)
 			{
diff --git a/Engine/Plugins/Experimental/NNERuntimeRDG/Source/NNERuntimeRDG/Private/NNERuntimeRDGDml.cpp b/Engine/Plugins/Experimental/NNERuntimeRDG/Source/NNERuntimeRDG/Private/NNERuntimeRDGDml.cpp
--- a/Engine/Plugins/Experimental/NNERuntimeRDG/Source/NNERuntimeRDG/Private/NNERuntimeRDGDml.cpp
+++ b/Engine/Plugins/Experimental/NNERuntimeRDG/Source/NNERuntimeRDG/Private/NNERuntimeRDGDml.cpp
@@ -108,7 +108,7 @@ TArray<uint8> UNNERuntimeRDGDmlImpl::CreateModelData(FString FileType, TConstArr
 		return {};
 	}
 
-	TUniquePtr<UE::NNECore::Internal::IModelOptimizer> Optimizer = UE::NNEUtils::Internal::CreateONNXToNNEModelOptimizer();
+	TUniquePtr<UE::NNE::Internal::IModelOptimizer> Optimizer = UE::NNEUtils::Internal::CreateONNXToNNEModelOptimizer();
 #ifdef NNE_USE_DIRECTML
 	Optimizer->AddValidator(MakeShared<FModelValidatorDml>());
 #endif
@@ -152,12 +152,12 @@ bool UNNERuntimeRDGDmlImpl::CanCreateModelRDG(TObjectPtr<UNNEModelData> ModelDat
 #endif
 };
 
-TUniquePtr<UE::NNECore::IModelRDG> UNNERuntimeRDGDmlImpl::CreateModel(TObjectPtr<UNNEModelData> ModelData)
+TUniquePtr<UE::NNE::IModelRDG> UNNERuntimeRDGDmlImpl::CreateModel(TObjectPtr<UNNEModelData> ModelData)
 {
 #ifdef NNE_USE_DIRECTML
 	if (!CanCreateModelRDG(ModelData))
 	{
-		return TUniquePtr<UE::NNECore::IModelRDG>();
+		return TUniquePtr<UE::NNE::IModelRDG>();
 	}
 
 	TConstArrayView<uint8> Data = ModelData->GetModelData(GetRuntimeName());
@@ -173,9 +173,9 @@ TUniquePtr<UE::NNECore::IModelRDG> UNNERuntimeRDGDmlImpl::CreateModel(TObjectPtr
 		FEngineAnalytics::GetProvider().RecordEvent(TEXT("NeuralNetworkEngine.CreateModel"), Attributes);
 	}
 
-	return TUniquePtr<UE::NNECore::IModelRDG>(Model);
+	return TUniquePtr<UE::NNE::IModelRDG>(Model);
 #else
-	return TUniquePtr<UE::NNECore::IModelRDG>();
+	return TUniquePtr<UE::NNE::IModelRDG>();
 #endif
 }
 
diff --git a/Engine/Plugins/Experimental/NNERuntimeRDG/Source/NNERuntimeRDG/Private/NNERuntimeRDGDml.h b/Engine/Plugins/Experimental/NNERuntimeRDG/Source/NNERuntimeRDG/Private/NNERuntimeRDGDml.h
--- a/Engine/Plugins/Experimental/NNERuntimeRDG/Source/NNERuntimeRDG/Private/NNERuntimeRDGDml.h
+++ b/Engine/Plugins/Experimental/NNERuntimeRDG/Source/NNERuntimeRDG/Private/NNERuntimeRDGDml.h
@@ -38,7 +38,7 @@ public:
 	virtual TArray<uint8> CreateModelData(FString FileType, TConstArrayView<uint8> FileData) override;
 
 	virtual bool CanCreateModelRDG(TObjectPtr<UNNEModelData> ModelData) const override;
-	virtual TUniquePtr<UE::NNECore::IModelRDG> CreateModel(TObjectPtr<UNNEModelData> ModelData) override;
+	virtual TUniquePtr<UE::NNE::IModelRDG> CreateModel(TObjectPtr<UNNEModelData> ModelData) override;
 private:
 	UE::NNERuntimeRDG::Private::Dml::FDmlDeviceContext* Ctx;
 };
diff --git a/Engine/Plugins/Experimental/NNERuntimeRDG/Source/NNERuntimeRDG/Private/NNERuntimeRDGHlsl.cpp b/Engine/Plugins/Experimental/NNERuntimeRDG/Source/NNERuntimeRDG/Private/NNERuntimeRDGHlsl.cpp
--- a/Engine/Plugins/Experimental/NNERuntimeRDG/Source/NNERuntimeRDG/Private/NNERuntimeRDGHlsl.cpp
+++ b/Engine/Plugins/Experimental/NNERuntimeRDG/Source/NNERuntimeRDG/Private/NNERuntimeRDGHlsl.cpp
@@ -100,7 +100,7 @@ TArray<uint8> UNNERuntimeRDGHlslImpl::CreateModelData(FString FileType, TConstAr
 		return {};
 	}
 
-	TUniquePtr<UE::NNECore::Internal::IModelOptimizer> Optimizer = UE::NNEUtils::Internal::CreateONNXToNNEModelOptimizer();
+	TUniquePtr<UE::NNE::Internal::IModelOptimizer> Optimizer = UE::NNEUtils::Internal::CreateONNXToNNEModelOptimizer();
 	Optimizer->AddValidator(MakeShared<FModelValidatorHlsl>());
 
 	FNNEModelRaw InputModel;
@@ -122,11 +122,11 @@ TArray<uint8> UNNERuntimeRDGHlslImpl::CreateModelData(FString FileType, TConstAr
 	return Result;
 };
 
-TUniquePtr<UE::NNECore::IModelRDG> UNNERuntimeRDGHlslImpl::CreateModel(TObjectPtr<UNNEModelData> ModelData)
+TUniquePtr<UE::NNE::IModelRDG> UNNERuntimeRDGHlslImpl::CreateModel(TObjectPtr<UNNEModelData> ModelData)
 {
 	if (!CanCreateModelRDG(ModelData))
 	{
-		return TUniquePtr<UE::NNECore::IModelRDG>();
+		return TUniquePtr<UE::NNE::IModelRDG>();
 	}
 
 	TConstArrayView<uint8> Data = ModelData->GetModelData(GetRuntimeName());
@@ -142,5 +142,5 @@ TUniquePtr<UE::NNECore::IModelRDG> UNNERuntimeRDGHlslImpl::CreateModel(TObjectPt
 		FEngineAnalytics::GetProvider().RecordEvent(TEXT("NeuralNetworkEngine.CreateModel"), Attributes);
 	}
 
-	return TUniquePtr<UE::NNECore::IModelRDG>(Model);
+	return TUniquePtr<UE::NNE::IModelRDG>(Model);
 }
diff --git a/Engine/Plugins/Experimental/NNERuntimeRDG/Source/NNERuntimeRDG/Private/NNERuntimeRDGHlsl.h b/Engine/Plugins/Experimental/NNERuntimeRDG/Source/NNERuntimeRDG/Private/NNERuntimeRDGHlsl.h
--- a/Engine/Plugins/Experimental/NNERuntimeRDG/Source/NNERuntimeRDG/Private/NNERuntimeRDGHlsl.h
+++ b/Engine/Plugins/Experimental/NNERuntimeRDG/Source/NNERuntimeRDG/Private/NNERuntimeRDGHlsl.h
@@ -33,5 +33,5 @@ public:
 	virtual TArray<uint8> CreateModelData(FString FileType, TConstArrayView<uint8> FileData) override;
 
 	virtual bool CanCreateModelRDG(TObjectPtr<UNNEModelData> ModelData) const override;
-	virtual TUniquePtr<UE::NNECore::IModelRDG> CreateModel(TObjectPtr<UNNEModelData> ModelData) override;
+	virtual TUniquePtr<UE::NNE::IModelRDG> CreateModel(TObjectPtr<UNNEModelData> ModelData) override;
 };
\ No newline at end of file
diff --git a/Engine/Plugins/Experimental/NNERuntimeRDG/Source/NNERuntimeRDG/Private/NNERuntimeRDGHlslHelper.cpp b/Engine/Plugins/Experimental/NNERuntimeRDG/Source/NNERuntimeRDG/Private/NNERuntimeRDGHlslHelper.cpp
--- a/Engine/Plugins/Experimental/NNERuntimeRDG/Source/NNERuntimeRDG/Private/NNERuntimeRDGHlslHelper.cpp
+++ b/Engine/Plugins/Experimental/NNERuntimeRDG/Source/NNERuntimeRDG/Private/NNERuntimeRDGHlslHelper.cpp
@@ -7,16 +7,16 @@
 
 namespace UE::NNERuntimeRDG::Private::Hlsl
 {
-	void FillTensorSizeShaderParameters(const NNECore::Internal::FTensor& Tensor, TStaticArray<FUintVector4, NXRT_TENSORSTRIDEINFO_MAX_NUM_DIMENSIONS, 16U>& OutShaderParam, int32 Idx)
+	void FillTensorSizeShaderParameters(const NNE::Internal::FTensor& Tensor, TStaticArray<FUintVector4, NXRT_TENSORSTRIDEINFO_MAX_NUM_DIMENSIONS, 16U>& OutShaderParam, int32 Idx)
 	{
-		static_assert(NNECore::FTensorShape::MaxRank <= NXRT_TENSORSTRIDEINFO_MAX_NUM_DIMENSIONS);
+		static_assert(NNE::FTensorShape::MaxRank <= NXRT_TENSORSTRIDEINFO_MAX_NUM_DIMENSIONS);
 		for (int32 i = 0; i < Tensor.GetShape().Rank(); ++i)
 		{
 			OutShaderParam[i][Idx] = Tensor.GetShape().GetData()[i];
 		}
 	}
 
-	void FillTensorStrideShaderParameters(const NNECore::Internal::FTensor& Tensor, TStaticArray<FUintVector4, NXRT_TENSORSTRIDEINFO_MAX_NUM_DIMENSIONS, 16U>& OutShaderParam, int32 Idx, int32 TargetNumdimensionForBroadcast)
+	void FillTensorStrideShaderParameters(const NNE::Internal::FTensor& Tensor, TStaticArray<FUintVector4, NXRT_TENSORSTRIDEINFO_MAX_NUM_DIMENSIONS, 16U>& OutShaderParam, int32 Idx, int32 TargetNumdimensionForBroadcast)
 	{
 		if (TargetNumdimensionForBroadcast == -1)
 		{
@@ -25,7 +25,7 @@ namespace UE::NNERuntimeRDG::Private::Hlsl
 		checkf(TargetNumdimensionForBroadcast >= Tensor.GetShape().Rank(), TEXT("Can't broadcast tensor from rank %d to rank %d, should be inferior or equal."), Tensor.GetShape().Rank(), TargetNumdimensionForBroadcast);
 		int32 Offset = TargetNumdimensionForBroadcast - Tensor.GetShape().Rank();
 
-		static_assert(NNECore::FTensorShape::MaxRank <= NXRT_TENSORSTRIDEINFO_MAX_NUM_DIMENSIONS);
+		static_assert(NNE::FTensorShape::MaxRank <= NXRT_TENSORSTRIDEINFO_MAX_NUM_DIMENSIONS);
 		for (int32 i = 7; i >= 0; --i)
 		{
 			if (i >= TargetNumdimensionForBroadcast || i < Offset)
@@ -43,7 +43,7 @@ namespace UE::NNERuntimeRDG::Private::Hlsl
 		}
 	}
 
-	void FillTensorStrideForBroadcastShaderParameters(const NNECore::Internal::FTensor& Tensor, int32 OutputNumdimension, TStaticArray<FUintVector4, NXRT_TENSORSTRIDEINFO_MAX_NUM_DIMENSIONS, 16U>& OutShaderParam, int32 Idx)
+	void FillTensorStrideForBroadcastShaderParameters(const NNE::Internal::FTensor& Tensor, int32 OutputNumdimension, TStaticArray<FUintVector4, NXRT_TENSORSTRIDEINFO_MAX_NUM_DIMENSIONS, 16U>& OutShaderParam, int32 Idx)
 	{
 		checkf(OutputNumdimension >= Tensor.GetShape().Rank(), TEXT("Can't broadcast tensor from rank %d to rank %d, should be inferior or equal."), Tensor.GetShape().Rank(), OutputNumdimension);
 		FillTensorStrideShaderParameters(Tensor, OutShaderParam, Idx, OutputNumdimension);
diff --git a/Engine/Plugins/Experimental/NNERuntimeRDG/Source/NNERuntimeRDG/Private/NNERuntimeRDGHlslHelper.h b/Engine/Plugins/Experimental/NNERuntimeRDG/Source/NNERuntimeRDG/Private/NNERuntimeRDGHlslHelper.h
--- a/Engine/Plugins/Experimental/NNERuntimeRDG/Source/NNERuntimeRDG/Private/NNERuntimeRDGHlslHelper.h
+++ b/Engine/Plugins/Experimental/NNERuntimeRDG/Source/NNERuntimeRDG/Private/NNERuntimeRDGHlslHelper.h
@@ -4,15 +4,15 @@
 
 #include "Containers/StaticArray.h"
 
-namespace UE::NNECore::Internal { class FTensor; }
+namespace UE::NNE::Internal { class FTensor; }
 
 #define NXRT_TENSORSTRIDEINFO_MAX_NUM_DIMENSIONS 8
 
 namespace UE::NNERuntimeRDG::Private::Hlsl
 {
 
-	void FillTensorSizeShaderParameters(const NNECore::Internal::FTensor& Tensor, TStaticArray<FUintVector4, NXRT_TENSORSTRIDEINFO_MAX_NUM_DIMENSIONS, 16U>& OutShaderParam, int32 Idx);
-	void FillTensorStrideShaderParameters(const NNECore::Internal::FTensor& Tensor, TStaticArray<FUintVector4, NXRT_TENSORSTRIDEINFO_MAX_NUM_DIMENSIONS, 16U>& OutShaderParam, int32 Idx, int32 TargetNumdimensionForBroadcast = -1);
-	void FillTensorStrideForBroadcastShaderParameters(const NNECore::Internal::FTensor& Tensor, int32 OutputNumdimension, TStaticArray<FUintVector4, NXRT_TENSORSTRIDEINFO_MAX_NUM_DIMENSIONS, 16U>& OutShaderParam, int32 Idx);
+	void FillTensorSizeShaderParameters(const NNE::Internal::FTensor& Tensor, TStaticArray<FUintVector4, NXRT_TENSORSTRIDEINFO_MAX_NUM_DIMENSIONS, 16U>& OutShaderParam, int32 Idx);
+	void FillTensorStrideShaderParameters(const NNE::Internal::FTensor& Tensor, TStaticArray<FUintVector4, NXRT_TENSORSTRIDEINFO_MAX_NUM_DIMENSIONS, 16U>& OutShaderParam, int32 Idx, int32 TargetNumdimensionForBroadcast = -1);
+	void FillTensorStrideForBroadcastShaderParameters(const NNE::Internal::FTensor& Tensor, int32 OutputNumdimension, TStaticArray<FUintVector4, NXRT_TENSORSTRIDEINFO_MAX_NUM_DIMENSIONS, 16U>& OutShaderParam, int32 Idx);
 	FIntVector ComputeElementWiseThreadGroups(uint32 ElementCount, uint32 GroupSizeX);
 } // namespace UE::NNERuntimeRDG::Private::Hlsl
\ No newline at end of file
diff --git a/Engine/Plugins/Experimental/NNERuntimeRDG/Source/NNERuntimeRDG/Private/NNERuntimeRDGHlslOp.h b/Engine/Plugins/Experimental/NNERuntimeRDG/Source/NNERuntimeRDG/Private/NNERuntimeRDGHlslOp.h
--- a/Engine/Plugins/Experimental/NNERuntimeRDG/Source/NNERuntimeRDG/Private/NNERuntimeRDGHlslOp.h
+++ b/Engine/Plugins/Experimental/NNERuntimeRDG/Source/NNERuntimeRDG/Private/NNERuntimeRDGHlslOp.h
@@ -12,7 +12,7 @@ namespace UE::NNERuntimeRDG::Private::Hlsl
 	struct FOperatorHlsl : public IOperatorRDG, public IPrepareOperator
 	{
 		virtual ~FOperatorHlsl() = default;
-		virtual bool Initialize(TConstArrayView<NNECore::FTensorDesc> InputTensorDescs, TConstArrayView<NNECore::FTensorDesc> OutputTensorDescs, const NNECore::FAttributeMap& Attributes) = 0;
+		virtual bool Initialize(TConstArrayView<NNE::FTensorDesc> InputTensorDescs, TConstArrayView<NNE::FTensorDesc> OutputTensorDescs, const NNE::FAttributeMap& Attributes) = 0;
 		virtual void Dispatch(FRDGBuilder& GraphBuilder, TConstArrayView<FTensorRDGRef> InInputTensors, TConstArrayView<FTensorRDGRef> InOutputTensors) = 0;
 		virtual void OptimizeInputsWeights(TArrayView<FTensorRDGRef> InputWeights) {};
 	};
diff --git a/Engine/Plugins/Experimental/NNERuntimeRDG/Source/NNERuntimeRDG/Private/NNERuntimeRDGModel.cpp b/Engine/Plugins/Experimental/NNERuntimeRDG/Source/NNERuntimeRDG/Private/NNERuntimeRDGModel.cpp
--- a/Engine/Plugins/Experimental/NNERuntimeRDG/Source/NNERuntimeRDG/Private/NNERuntimeRDGModel.cpp
+++ b/Engine/Plugins/Experimental/NNERuntimeRDG/Source/NNERuntimeRDG/Private/NNERuntimeRDGModel.cpp
@@ -39,8 +39,8 @@ bool FModelInstanceRDG::LoadModel(TConstArrayView<uint8> ModelData, FNNERuntimeF
 	{
 		const FNNEFormatTensorDesc& FormatTensorDesc = Format.Tensors[Idx];
 
-		const NNECore::FSymbolicTensorShape SymbolicShape = NNECore::FSymbolicTensorShape::Make(FormatTensorDesc.Shape);
-		const NNECore::FTensorDesc SymbolicTensor = NNECore::FTensorDesc::Make(FormatTensorDesc.Name, SymbolicShape, FormatTensorDesc.DataType);
+		const NNE::FSymbolicTensorShape SymbolicShape = NNE::FSymbolicTensorShape::Make(FormatTensorDesc.Shape);
+		const NNE::FTensorDesc SymbolicTensor = NNE::FTensorDesc::Make(FormatTensorDesc.Name, SymbolicShape, FormatTensorDesc.DataType);
 
 		AllSymbolicTensorDescs.Emplace(SymbolicTensor);
 
@@ -67,7 +67,7 @@ bool FModelInstanceRDG::LoadModel(TConstArrayView<uint8> ModelData, FNNERuntimeF
 				return false;
 			}
 
-			const NNECore::FTensorShape TensorShape = NNECore::FTensorShape::MakeFromSymbolic(SymbolicTensor.GetShape());
+			const NNE::FTensorShape TensorShape = NNE::FTensorShape::MakeFromSymbolic(SymbolicTensor.GetShape());
 			FTensorRDG& WeightRDG = WeightTensorRDGs.Emplace_GetRef(FTensorRDG::Make(SymbolicTensor, TensorShape, nullptr));
 
 			if (WeightRDG.GetDataSize() != FormatTensorDesc.DataSize)
@@ -96,12 +96,12 @@ bool FModelInstanceRDG::LoadModel(TConstArrayView<uint8> ModelData, FNNERuntimeF
 
 
 
-int FModelInstanceRDG::SetInputTensorShapes(TConstArrayView<NNECore::FTensorShape> InInputShapes)
+int FModelInstanceRDG::SetInputTensorShapes(TConstArrayView<NNE::FTensorShape> InInputShapes)
 {
 	OutputTensorShapes.Reset(OutputTensorIndices.Num());
 
 	//Verify input shape are valid for the model and set InputTensorShapes
-	if (FModelInstanceBase<NNECore::IModelInstanceRDG>::SetInputTensorShapes(InInputShapes) != 0)
+	if (FModelInstanceBase<NNE::IModelInstanceRDG>::SetInputTensorShapes(InInputShapes) != 0)
 	{
 		return -1;
 	}
@@ -113,8 +113,8 @@ int FModelInstanceRDG::SetInputTensorShapes(TConstArrayView<NNECore::FTensorShap
 	for (int32 i = 0; i < InputTensorIndices.Num(); ++i)
 	{
 		const int32 Idx = InputTensorIndices[i];
-		const NNECore::FTensorDesc& TensorDesc = InputSymbolicTensors[i];
-		const NNECore::FTensorShape& TensorShape = InputTensorShapes[i];
+		const NNE::FTensorDesc& TensorDesc = InputSymbolicTensors[i];
+		const NNE::FTensorShape& TensorShape = InputTensorShapes[i];
 
 		InputTensorRDGs.Emplace(FTensorRDG::Make(TensorDesc, TensorShape, nullptr));
 		AllTensorRDGRefs[Idx] = &InputTensorRDGs[i];
@@ -131,8 +131,8 @@ int FModelInstanceRDG::SetInputTensorShapes(TConstArrayView<NNECore::FTensorShap
 	for (int32 i = 0; i < IntermediateTensorIndices.Num(); ++i)
 	{
 		const int32 Idx = IntermediateTensorIndices[i];
-		const NNECore::FTensorDesc& TensorDesc = AllSymbolicTensorDescs[Idx];
-		const NNECore::FTensorShape TensorShape = NNECore::FTensorShape::MakeFromSymbolic(TensorDesc.GetShape());
+		const NNE::FTensorDesc& TensorDesc = AllSymbolicTensorDescs[Idx];
+		const NNE::FTensorShape TensorShape = NNE::FTensorShape::MakeFromSymbolic(TensorDesc.GetShape());
 
 		IntermediateTensorRDGs.Emplace(FTensorRDG::Make(TensorDesc, TensorShape, nullptr));
 		AllTensorRDGRefs[Idx] = &IntermediateTensorRDGs[i];
@@ -142,8 +142,8 @@ int FModelInstanceRDG::SetInputTensorShapes(TConstArrayView<NNECore::FTensorShap
 	for (int32 i = 0; i < OutputTensorIndices.Num(); ++i)
 	{
 		const int32 Idx = OutputTensorIndices[i];
-		const NNECore::FTensorDesc& TensorDesc = OutputSymbolicTensors[i];
-		const NNECore::FTensorShape TensorShape = NNECore::FTensorShape::MakeFromSymbolic(TensorDesc.GetShape());
+		const NNE::FTensorDesc& TensorDesc = OutputSymbolicTensors[i];
+		const NNE::FTensorShape TensorShape = NNE::FTensorShape::MakeFromSymbolic(TensorDesc.GetShape());
 
 		OutputTensorRDGs.Emplace(FTensorRDG::Make(TensorDesc, TensorShape, nullptr));
 		AllTensorRDGRefs[Idx] = &OutputTensorRDGs[i];
@@ -197,7 +197,7 @@ FRDGBufferDesc CreateRDGBufferDescForTensorRDG(const FTensorRDG& Tensor)
 /**
  * Enqueue operators to RDG, the caller will run the GraphBuilder.Execute()
  */
-int FModelInstanceRDG::EnqueueRDG(FRDGBuilder& RDGBuilder, TConstArrayView<NNECore::FTensorBindingRDG> InInputBindings, TConstArrayView<NNECore::FTensorBindingRDG> InOutputBindings)
+int FModelInstanceRDG::EnqueueRDG(FRDGBuilder& RDGBuilder, TConstArrayView<NNE::FTensorBindingRDG> InInputBindings, TConstArrayView<NNE::FTensorBindingRDG> InOutputBindings)
 {
 	check(IsInRenderingThread());
 
@@ -253,14 +253,14 @@ int FModelInstanceRDG::EnqueueRDG(FRDGBuilder& RDGBuilder, TConstArrayView<NNECo
 	return 0;
 }
 
-int FModelInstanceRDG::SetTensors(FRDGBuilder& GraphBuilder, FTensorRDGArray& InTensorRDGs, TConstArrayView<NNECore::FTensorBindingRDG> InBindings)
+int FModelInstanceRDG::SetTensors(FRDGBuilder& GraphBuilder, FTensorRDGArray& InTensorRDGs, TConstArrayView<NNE::FTensorBindingRDG> InBindings)
 {
 	check(InBindings.Num() == InTensorRDGs.Num());
 	
 	for (int32 Idx = 0; Idx < InBindings.Num(); ++Idx)
 	{
 		FTensorRDG& TensorRDG = InTensorRDGs[Idx];
-		const NNECore::FTensorBindingRDG& Binding = InBindings[Idx];
+		const NNE::FTensorBindingRDG& Binding = InBindings[Idx];
 		if (Binding.Buffer == nullptr)
 		{
 			return Idx;
diff --git a/Engine/Plugins/Experimental/NNERuntimeRDG/Source/NNERuntimeRDG/Private/NNERuntimeRDGModel.h b/Engine/Plugins/Experimental/NNERuntimeRDG/Source/NNERuntimeRDG/Private/NNERuntimeRDGModel.h
--- a/Engine/Plugins/Experimental/NNERuntimeRDG/Source/NNERuntimeRDG/Private/NNERuntimeRDGModel.h
+++ b/Engine/Plugins/Experimental/NNERuntimeRDG/Source/NNERuntimeRDG/Private/NNERuntimeRDGModel.h
@@ -11,25 +11,25 @@
 
 namespace UE::NNERuntimeRDG::Private
 {
-	class FModelInstanceRDG : public NNECore::Internal::FModelInstanceBase<NNECore::IModelInstanceRDG>
+	class FModelInstanceRDG : public NNE::Internal::FModelInstanceBase<NNE::IModelInstanceRDG>
 	{
 	public:
 		FModelInstanceRDG() {};
 		virtual ~FModelInstanceRDG() = default;
 
-		virtual int SetInputTensorShapes(TConstArrayView<NNECore::FTensorShape> InInputShapes) override;
-		virtual int EnqueueRDG(FRDGBuilder& RDGBuilder, TConstArrayView<NNECore::FTensorBindingRDG> InInputBindings, TConstArrayView<NNECore::FTensorBindingRDG> InOutputBindings) override;
+		virtual int SetInputTensorShapes(TConstArrayView<NNE::FTensorShape> InInputShapes) override;
+		virtual int EnqueueRDG(FRDGBuilder& RDGBuilder, TConstArrayView<NNE::FTensorBindingRDG> InInputBindings, TConstArrayView<NNE::FTensorBindingRDG> InOutputBindings) override;
 		
 	protected:
 		bool LoadModel(TConstArrayView<uint8> ModelData, FNNERuntimeFormat& Format, int32 GuidAndVersionSize);
-		int SetTensors(FRDGBuilder& GraphBuilder, FTensorRDGArray& InTensorRDGs, TConstArrayView<NNECore::FTensorBindingRDG> InBindings);
+		int SetTensors(FRDGBuilder& GraphBuilder, FTensorRDGArray& InTensorRDGs, TConstArrayView<NNE::FTensorBindingRDG> InBindings);
 
 		virtual int PrepareTensorShapesAndData() = 0;
 		virtual bool PrepareModelRDG(FRDGBuilder& RDGBuilder) { return false; }
 		virtual void AddDispatchOps_RenderThread(FRDGBuilder& GraphBuilder) = 0;
 
 		//Tensor descriptor
-		TArray<NNECore::FTensorDesc>	AllSymbolicTensorDescs;
+		TArray<NNE::FTensorDesc>	AllSymbolicTensorDescs;
 
 		//Tensor indices for models
 		TArray<int32>				IntermediateTensorIndices;
diff --git a/Engine/Plugins/Experimental/NNERuntimeRDG/Source/NNERuntimeRDG/Private/NNERuntimeRDGModelHlsl.cpp b/Engine/Plugins/Experimental/NNERuntimeRDG/Source/NNERuntimeRDG/Private/NNERuntimeRDGModelHlsl.cpp
--- a/Engine/Plugins/Experimental/NNERuntimeRDG/Source/NNERuntimeRDG/Private/NNERuntimeRDGModelHlsl.cpp
+++ b/Engine/Plugins/Experimental/NNERuntimeRDG/Source/NNERuntimeRDG/Private/NNERuntimeRDGModelHlsl.cpp
@@ -13,7 +13,7 @@ namespace UE::NNERuntimeRDG::Private::Hlsl
 namespace ModelUtils
 {
 
-FOperatorHlsl* OpCreate(const FString& OpName, TConstArrayView<NNECore::FTensorDesc> InputTensorDescs, TConstArrayView<NNECore::FTensorDesc> OutputTensorDescs, const NNECore::FAttributeMap& AttributeMap)
+FOperatorHlsl* OpCreate(const FString& OpName, TConstArrayView<NNE::FTensorDesc> InputTensorDescs, TConstArrayView<NNE::FTensorDesc> OutputTensorDescs, const NNE::FAttributeMap& AttributeMap)
 {
 	FOperatorRegistryHlsl::OperatorCreateFunc CreateFn = FOperatorRegistryHlsl::Get()->OpFind(OpName);
 
@@ -116,8 +116,8 @@ bool FModelInstance::Init(TConstArrayView<uint8> ModelData)
 		return false;
 	}
 
-	TArray<NNECore::FTensorDesc> Inputs;
-	TArray<NNECore::FTensorDesc> Outputs;
+	TArray<NNE::FTensorDesc> Inputs;
+	TArray<NNE::FTensorDesc> Outputs;
 	TArray<FTensorRDGRef> InputsAsWeights;
 
 	// Loop over all operators in the model and create them
@@ -127,7 +127,7 @@ bool FModelInstance::Init(TConstArrayView<uint8> ModelData)
 		Inputs.Reset();
 		Outputs.Reset();
 		InputsAsWeights.Reset();
-		UE::NNECore::FAttributeMap AttributeMap;
+		UE::NNE::FAttributeMap AttributeMap;
 
 		for (int32 InputTensorIndex : Format.Operators[Idx].InTensors)
 		{
@@ -220,8 +220,8 @@ int FModelInstance::PrepareTensorShapesAndData()
 	// Run model preparation (including shape inference) on all operators
 	// This loop could be abstracted to a different runtime/system as it apply on FTensorRef & IPrepareOperator witch are RDG agnostics.
 	static constexpr int32 MaxExpectedInput = 10;
-	TArray<NNECore::Internal::FTensorRef, TInlineAllocator<MaxExpectedInput>> InputTensors;
-	TArray<NNECore::Internal::FTensorRef> OutputTensors;
+	TArray<NNE::Internal::FTensorRef, TInlineAllocator<MaxExpectedInput>> InputTensors;
+	TArray<NNE::Internal::FTensorRef> OutputTensors;
 	TArray<bool> AllInitializedTensors;
 
 	checkCode(
@@ -364,18 +364,18 @@ bool FModelInstance::PrepareWeights()
 	return true;
 }
 
-TUniquePtr<UE::NNECore::IModelInstanceRDG> FModel::CreateModelInstance()
+TUniquePtr<UE::NNE::IModelInstanceRDG> FModel::CreateModelInstance()
 {
 	FModelInstance* ModelInstance = new FModelInstance();
 
 	if (!ModelInstance->Init(ModelData))
 	{
 		delete ModelInstance;
-		return TUniquePtr<UE::NNECore::IModelInstanceRDG>();
+		return TUniquePtr<UE::NNE::IModelInstanceRDG>();
 	}
 
-	UE::NNECore::IModelInstanceRDG* IModelInstance = static_cast<UE::NNECore::IModelInstanceRDG*>(ModelInstance);
-	return TUniquePtr<UE::NNECore::IModelInstanceRDG>(IModelInstance);
+	UE::NNE::IModelInstanceRDG* IModelInstance = static_cast<UE::NNE::IModelInstanceRDG*>(ModelInstance);
+	return TUniquePtr<UE::NNE::IModelInstanceRDG>(IModelInstance);
 }
 
 FModel::FModel(TConstArrayView<uint8> InModelData) : ModelData(InModelData) {}
diff --git a/Engine/Plugins/Experimental/NNERuntimeRDG/Source/NNERuntimeRDG/Private/NNERuntimeRDGModelHlsl.h b/Engine/Plugins/Experimental/NNERuntimeRDG/Source/NNERuntimeRDG/Private/NNERuntimeRDGModelHlsl.h
--- a/Engine/Plugins/Experimental/NNERuntimeRDG/Source/NNERuntimeRDG/Private/NNERuntimeRDGModelHlsl.h
+++ b/Engine/Plugins/Experimental/NNERuntimeRDG/Source/NNERuntimeRDG/Private/NNERuntimeRDGModelHlsl.h
@@ -35,13 +35,13 @@ private:
 	TArray<TRefCountPtr<FRDGPooledBuffer>> ConstantsExternalRDGResources;
 };
 
-class FModel : public NNECore::IModelRDG
+class FModel : public NNE::IModelRDG
 {
 public:
 	FModel(TConstArrayView<uint8> ModelData);
 	virtual ~FModel() {};
 
-	virtual TUniquePtr<UE::NNECore::IModelInstanceRDG> CreateModelInstance() override;
+	virtual TUniquePtr<UE::NNE::IModelInstanceRDG> CreateModelInstance() override;
 
 private:
 	TArray<uint8> ModelData;
