NNE Rename NNECore module to NNE following split of runtimes into there own plugins (step2)

diff --git a/Engine/Plugins/Experimental/NNERuntimeRDG/Source/NNERuntimeRDG/Private/NNERuntimeRDGModule.cpp b/Engine/Plugins/Experimental/NNERuntimeRDG/Source/NNERuntimeRDG/Private/NNERuntimeRDGModule.cpp
--- a/Engine/Plugins/Experimental/NNERuntimeRDG/Source/NNERuntimeRDG/Private/NNERuntimeRDGModule.cpp
+++ b/Engine/Plugins/Experimental/NNERuntimeRDG/Source/NNERuntimeRDG/Private/NNERuntimeRDGModule.cpp
@@ -17,7 +17,7 @@ void FNNERuntimeRDGModule::StartupModule()
 
 		NNERuntimeRDGHlsl->Init();
 		NNERuntimeRDGHlsl->AddToRoot();
-		UE::NNECore::RegisterRuntime(RuntimeCPUInterface);
+		UE::NNE::RegisterRuntime(RuntimeCPUInterface);
 	}
 
 #ifdef NNE_USE_DIRECTML
@@ -30,7 +30,7 @@ void FNNERuntimeRDGModule::StartupModule()
 		bool bRegisterOnlyOperators = UE::NNERuntimeRDG::Private::Dml::FRuntimeDmlStartup();
 		NNERuntimeRDGDml->Init(bRegisterOnlyOperators);
 		NNERuntimeRDGDml->AddToRoot();
-		UE::NNECore::RegisterRuntime(RuntimeDmlInterface);
+		UE::NNE::RegisterRuntime(RuntimeDmlInterface);
 	}
 #endif
 }
@@ -42,7 +42,7 @@ void FNNERuntimeRDGModule::ShutdownModule()
 	{
 		TWeakInterfacePtr<INNERuntime> RuntimeCPUInterface(NNERuntimeRDGHlsl.Get());
 
-		UE::NNECore::UnregisterRuntime(RuntimeCPUInterface);
+		UE::NNE::UnregisterRuntime(RuntimeCPUInterface);
 		NNERuntimeRDGHlsl->RemoveFromRoot();
 		NNERuntimeRDGHlsl = TWeakObjectPtr<UNNERuntimeRDGHlslImpl>(nullptr);
 	}
@@ -53,7 +53,7 @@ void FNNERuntimeRDGModule::ShutdownModule()
 	{
 		TWeakInterfacePtr<INNERuntime> RuntimeDmlInterface(NNERuntimeRDGDml.Get());
 
-		UE::NNECore::UnregisterRuntime(RuntimeDmlInterface);
+		UE::NNE::UnregisterRuntime(RuntimeDmlInterface);
 		NNERuntimeRDGDml->RemoveFromRoot();
 		NNERuntimeRDGDml = TWeakObjectPtr<UNNERuntimeRDGDmlImpl>(nullptr);
 	}
