[NNE] Implement RDG Dml validation for activations

diff --git a/Engine/Plugins/Experimental/NNERuntimeRDG/Source/NNERuntimeRDG/Private/Dml/NNEDmlOperatorActivations.cpp b/Engine/Plugins/Experimental/NNERuntimeRDG/Source/NNERuntimeRDG/Private/Dml/NNEDmlOperatorActivations.cpp
--- a/Engine/Plugins/Experimental/NNERuntimeRDG/Source/NNERuntimeRDG/Private/Dml/NNEDmlOperatorActivations.cpp
+++ b/Engine/Plugins/Experimental/NNERuntimeRDG/Source/NNERuntimeRDG/Private/Dml/NNEDmlOperatorActivations.cpp
@@ -25,7 +25,40 @@ public:
 
 	static bool Validate(const NNE::FAttributeMap& AttributeMap, TConstArrayView<ENNETensorDataType> InputTypes, TConstArrayView<NNE::FSymbolicTensorShape> InputShapes)
 	{
-		//TODO
+		if(InputShapes.Num() != 1)
+		{
+			UE_LOG(LogNNE, Warning, TEXT("Invalid number of input tensors"));
+			return false;
+		}
+
+		if (!(InputTypes[0] == ENNETensorDataType::Half || InputTypes[0] == ENNETensorDataType::Float))
+		{
+			UE_LOG(LogNNE, Warning, TEXT("Invalid DML tensor data type"));
+			return false;
+		}
+
+		const NNE::FSymbolicTensorShape& InputShape = InputShapes[0];
+
+		if(!InputShape.IsConcrete())
+		{
+			UE_LOG(LogNNE, Warning, TEXT("DML tensor shape must be concrete"));
+			return false;
+		}
+
+		if (NNE::FTensorShape::MakeFromSymbolic(InputShape).Volume() == 0)
+		{
+			UE_LOG(LogNNE, Warning, TEXT("Invalid DML tensor size, it's 0"));
+			return false;
+		}
+
+		const int32 MinTensorRank(1), MaxTensorRank(DML_TENSOR_DIMENSION_COUNT_MAX1);
+
+		if (InputShape.Rank() < MinTensorRank || InputShape.Rank() > MaxTensorRank)
+		{
+			UE_LOG(LogNNE, Warning, TEXT("Invalid DML tensor rank:%d [%d,%d]"), InputShape.Rank(), MinTensorRank, MaxTensorRank);
+			return false;
+		}
+
 		return true;
 	}
 
@@ -189,7 +222,56 @@ public:
 
 	static bool Validate(const NNE::FAttributeMap& AttributeMap, TConstArrayView<ENNETensorDataType> InputTypes, TConstArrayView<NNE::FSymbolicTensorShape> InputShapes)
 	{
-		//TODO
+		if(InputShapes.Num() != 2)
+		{
+			UE_LOG(LogNNE, Warning, TEXT("Invalid number of input tensors"));
+			return false;
+		}
+
+		if ( !(InputTypes[0] == ENNETensorDataType::Half || InputTypes[0] == ENNETensorDataType::Float) ||
+			 !(InputTypes[1] == ENNETensorDataType::Half || InputTypes[1] == ENNETensorDataType::Float) )
+		{
+			UE_LOG(LogNNE, Warning, TEXT("Invalid DML tensor data type"));
+			return false;
+		}
+
+		const NNE::FSymbolicTensorShape& InputShape = InputShapes[0];
+
+		if(!InputShape.IsConcrete())
+		{
+			UE_LOG(LogNNE, Warning, TEXT("DML tensor shape must be concrete"));
+			return false;
+		}
+
+		if (NNE::FTensorShape::MakeFromSymbolic(InputShape).Volume() == 0)
+		{
+			UE_LOG(LogNNE, Warning, TEXT("Invalid DML tensor size, it's 0"));
+			return false;
+		}
+
+		const int32 MinTensorRank(1), MaxTensorRank(DML_TENSOR_DIMENSION_COUNT_MAX1);
+
+		if (InputShape.Rank() < MinTensorRank || InputShape.Rank() > MaxTensorRank)
+		{
+			UE_LOG(LogNNE, Warning, TEXT("Invalid DML tensor rank:%d [%d,%d]"), InputShape.Rank(), MinTensorRank, MaxTensorRank);
+			return false;
+		}
+
+		const NNE::FSymbolicTensorShape& Input2Shape = InputShapes[1];
+
+		if(!Input2Shape.IsConcrete())
+		{
+			UE_LOG(LogNNE, Warning, TEXT("DML input shape must be concrete"));
+			return false;
+		}
+
+		if(!IsEqualOrBroadcastable(InputShape.GetData(), Input2Shape.GetData()))
+		{
+
+			UE_LOG(LogNNE, Warning, TEXT("DML input shapes must be compatible"));
+			return false;
+		}
+
 		return true;
 	}
 
@@ -272,7 +354,7 @@ struct FDmlOperator##OpName##Registrator \
 { \
 	FDmlOperator##OpName##Registrator() \
 	{ \
-		FOperatorRegistryDml::Get()->OpAdd(TEXT(#OpName), OpClass<DML_ACTIVATION_##DmlOpName##_OPERATOR_DESC, DML_OPERATOR_ACTIVATION_##DmlOpName>::Create); \
+		FOperatorRegistryDml::Get()->OpAdd(TEXT(#OpName), OpClass<DML_ACTIVATION_##DmlOpName##_OPERATOR_DESC, DML_OPERATOR_ACTIVATION_##DmlOpName>::Create, OpClass<DML_ACTIVATION_##DmlOpName##_OPERATOR_DESC, DML_OPERATOR_ACTIVATION_##DmlOpName>::Validate); \
 	} \
 }; \
 \
diff --git a/Engine/Plugins/Experimental/NNERuntimeRDG/Source/NNERuntimeRDG/Private/Dml/NNEDmlOperatorElementwise.cpp b/Engine/Plugins/Experimental/NNERuntimeRDG/Source/NNERuntimeRDG/Private/Dml/NNEDmlOperatorElementwise.cpp
--- a/Engine/Plugins/Experimental/NNERuntimeRDG/Source/NNERuntimeRDG/Private/Dml/NNEDmlOperatorElementwise.cpp
+++ b/Engine/Plugins/Experimental/NNERuntimeRDG/Source/NNERuntimeRDG/Private/Dml/NNEDmlOperatorElementwise.cpp
@@ -133,6 +133,12 @@ public:
 		return new FOperatorDmlElementWiseBinary();
 	}
 
+	static bool Validate(const NNE::FAttributeMap& AttributeMap, TConstArrayView<ENNETensorDataType> InputTypes, TConstArrayView<NNE::FSymbolicTensorShape> InputShapes)
+	{
+		//TODO
+		return true;
+	}
+
 private:
 
 	FOperatorDmlElementWiseBinary() = default;
@@ -212,7 +218,7 @@ struct FDmlOperator##OpName##Registrator \
 { \
 	FDmlOperator##OpName##Registrator() \
 	{ \
-		FOperatorRegistryDml::Get()->OpAdd(TEXT(#OpName), OpClass<DML_ELEMENT_WISE_##DmlOpName##_OPERATOR_DESC, DML_OPERATOR_ELEMENT_WISE_##DmlOpName>::Create); \
+		FOperatorRegistryDml::Get()->OpAdd(TEXT(#OpName), OpClass<DML_ELEMENT_WISE_##DmlOpName##_OPERATOR_DESC, DML_OPERATOR_ELEMENT_WISE_##DmlOpName>::Create, OpClass<DML_ELEMENT_WISE_##DmlOpName##_OPERATOR_DESC, DML_OPERATOR_ELEMENT_WISE_##DmlOpName>::Validate); \
 	} \
 }; \
 \
diff --git a/Engine/Plugins/Experimental/NNERuntimeRDG/Source/NNERuntimeRDG/Private/Dml/NNEDmlOperatorUtils.h b/Engine/Plugins/Experimental/NNERuntimeRDG/Source/NNERuntimeRDG/Private/Dml/NNEDmlOperatorUtils.h
--- a/Engine/Plugins/Experimental/NNERuntimeRDG/Source/NNERuntimeRDG/Private/Dml/NNEDmlOperatorUtils.h
+++ b/Engine/Plugins/Experimental/NNERuntimeRDG/Source/NNERuntimeRDG/Private/Dml/NNEDmlOperatorUtils.h
@@ -39,6 +39,35 @@ static EAutoPad AutoPadFromString(FStringView StringVal)
     }
 }
 
+template<typename T>
+static bool IsEqualOrBroadcastable(TConstArrayView<T> ShapeA, TConstArrayView<T> ShapeB)
+{
+    if(ShapeA.Num() < ShapeB.Num())
+    {
+        return false;
+    }
+
+    int32 IdxOffset = 0;
+
+    if(ShapeB.Num() < ShapeA.Num())
+    {
+        IdxOffset = ShapeA.Num() - ShapeB.Num();
+    }
+
+    for(int32 Idx = 0; Idx < ShapeB.Num(); ++Idx)
+    {
+        if(ShapeA[Idx + IdxOffset] != ShapeB[Idx])
+        {
+            if(ShapeB[Idx] != 1)
+            {
+                return false;
+            }
+        }
+    }
+
+    return true;
+}
+
 static Util::FSmallUIntArray KernelPadding(
     TConstArrayView<uint32> InputShape, TConstArrayView<uint32> WindowSize, 
     TConstArrayView<uint32> Dilations, TConstArrayView<uint32> Strides
