[NNE] Unify validations of activations and elementwise

#rnx

diff --git a/Engine/Plugins/Experimental/NNERuntimeRDG/Source/NNERuntimeRDG/Private/Dml/NNEDmlOperatorActivations.cpp b/Engine/Plugins/Experimental/NNERuntimeRDG/Source/NNERuntimeRDG/Private/Dml/NNEDmlOperatorActivations.cpp
--- a/Engine/Plugins/Experimental/NNERuntimeRDG/Source/NNERuntimeRDG/Private/Dml/NNEDmlOperatorActivations.cpp
+++ b/Engine/Plugins/Experimental/NNERuntimeRDG/Source/NNERuntimeRDG/Private/Dml/NNEDmlOperatorActivations.cpp
@@ -30,36 +30,7 @@ public:
 			UE_LOG(LogNNE, Warning, TEXT("Invalid number of input tensors"));
 			return false;
 		}
-
-		if (!(InputTypes[0] == ENNETensorDataType::Half || InputTypes[0] == ENNETensorDataType::Float))
-		{
-			UE_LOG(LogNNE, Warning, TEXT("Invalid DML tensor data type"));
-			return false;
-		}
-
-		const NNE::FSymbolicTensorShape& InputShape = InputShapes[0];
-
-		if(!InputShape.IsConcrete())
-		{
-			UE_LOG(LogNNE, Warning, TEXT("DML tensor shape must be concrete"));
-			return false;
-		}
-
-		if (NNE::FTensorShape::MakeFromSymbolic(InputShape).Volume() == 0)
-		{
-			UE_LOG(LogNNE, Warning, TEXT("Invalid DML tensor size, it's 0"));
-			return false;
-		}
-
-		const int32 MinTensorRank(1), MaxTensorRank(DML_TENSOR_DIMENSION_COUNT_MAX1);
-
-		if (InputShape.Rank() < MinTensorRank || InputShape.Rank() > MaxTensorRank)
-		{
-			UE_LOG(LogNNE, Warning, TEXT("Invalid DML tensor rank:%d [%d,%d]"), InputShape.Rank(), MinTensorRank, MaxTensorRank);
-			return false;
-		}
-
-		return true;
+		return CheckElementwiseTensor(InputTypes[0], InputShapes[0]);
 	}
 
 	virtual ~FOperatorDmlActivationUnary() = default;
@@ -228,47 +199,13 @@ public:
 			return false;
 		}
 
-		if ( !(InputTypes[0] == ENNETensorDataType::Half || InputTypes[0] == ENNETensorDataType::Float) ||
-			 !(InputTypes[1] == ENNETensorDataType::Half || InputTypes[1] == ENNETensorDataType::Float) )
+		if(!CheckElementwiseTensor(InputTypes[0], InputShapes[0]))
 		{
-			UE_LOG(LogNNE, Warning, TEXT("Invalid DML tensor data type"));
 			return false;
 		}
 
-		const NNE::FSymbolicTensorShape& InputShape = InputShapes[0];
-
-		if(!InputShape.IsConcrete())
-		{
-			UE_LOG(LogNNE, Warning, TEXT("DML tensor shape must be concrete"));
-			return false;
-		}
-
-		if (NNE::FTensorShape::MakeFromSymbolic(InputShape).Volume() == 0)
-		{
-			UE_LOG(LogNNE, Warning, TEXT("Invalid DML tensor size, it's 0"));
-			return false;
-		}
-
-		const int32 MinTensorRank(1), MaxTensorRank(DML_TENSOR_DIMENSION_COUNT_MAX1);
-
-		if (InputShape.Rank() < MinTensorRank || InputShape.Rank() > MaxTensorRank)
+		if(!CheckElementwiseTensor(InputTypes[1], InputShapes[1]))
 		{
-			UE_LOG(LogNNE, Warning, TEXT("Invalid DML tensor rank:%d [%d,%d]"), InputShape.Rank(), MinTensorRank, MaxTensorRank);
-			return false;
-		}
-
-		const NNE::FSymbolicTensorShape& Input2Shape = InputShapes[1];
-
-		if(!Input2Shape.IsConcrete())
-		{
-			UE_LOG(LogNNE, Warning, TEXT("DML input shape must be concrete"));
-			return false;
-		}
-
-		if(!IsEqualOrBroadcastable(InputShape.GetData(), Input2Shape.GetData()))
-		{
-
-			UE_LOG(LogNNE, Warning, TEXT("DML input shapes must be compatible"));
 			return false;
 		}
 
diff --git a/Engine/Plugins/Experimental/NNERuntimeRDG/Source/NNERuntimeRDG/Private/Dml/NNEDmlOperatorElementwise.cpp b/Engine/Plugins/Experimental/NNERuntimeRDG/Source/NNERuntimeRDG/Private/Dml/NNEDmlOperatorElementwise.cpp
--- a/Engine/Plugins/Experimental/NNERuntimeRDG/Source/NNERuntimeRDG/Private/Dml/NNEDmlOperatorElementwise.cpp
+++ b/Engine/Plugins/Experimental/NNERuntimeRDG/Source/NNERuntimeRDG/Private/Dml/NNEDmlOperatorElementwise.cpp
@@ -1,6 +1,7 @@
 // Copyright Epic Games, Inc. All Rights Reserved.
 #ifdef NNE_USE_DIRECTML
 #include "NNEDmlOperator.h"
+#include "NNEDmlOperatorUtils.h"
 
 namespace UE::NNERuntimeRDG::Private::Dml
 {
@@ -25,7 +26,17 @@ public:
 
 	static bool Validate(const NNE::FAttributeMap& AttributeMap, TConstArrayView<ENNETensorDataType> InputTypes, TConstArrayView<NNE::FSymbolicTensorShape> InputShapes)
 	{
-		//TODO
+		if(InputShapes.Num() != 1)
+		{
+			UE_LOG(LogNNE, Warning, TEXT("Invalid number of input tensors"));
+			return false;
+		}
+
+		if(!CheckElementwiseTensor(InputTypes[0], InputShapes[0]))
+		{
+			return false;
+		}
+
 		return true;
 	}
 
@@ -135,7 +146,22 @@ public:
 
 	static bool Validate(const NNE::FAttributeMap& AttributeMap, TConstArrayView<ENNETensorDataType> InputTypes, TConstArrayView<NNE::FSymbolicTensorShape> InputShapes)
 	{
-		//TODO
+		if(InputShapes.Num() != 2)
+		{
+			UE_LOG(LogNNE, Warning, TEXT("Invalid number of input tensors"));
+			return false;
+		}
+
+		if(!CheckElementwiseTensor(InputTypes[0], InputShapes[0]))
+		{
+			return false;
+		}
+
+		if(!CheckElementwiseTensor(InputTypes[1], InputShapes[1]))
+		{
+			return false;
+		}
+
 		return true;
 	}
 
diff --git a/Engine/Plugins/Experimental/NNERuntimeRDG/Source/NNERuntimeRDG/Private/Dml/NNEDmlOperatorUtils.h b/Engine/Plugins/Experimental/NNERuntimeRDG/Source/NNERuntimeRDG/Private/Dml/NNEDmlOperatorUtils.h
--- a/Engine/Plugins/Experimental/NNERuntimeRDG/Source/NNERuntimeRDG/Private/Dml/NNEDmlOperatorUtils.h
+++ b/Engine/Plugins/Experimental/NNERuntimeRDG/Source/NNERuntimeRDG/Private/Dml/NNEDmlOperatorUtils.h
@@ -39,30 +39,32 @@ static EAutoPad AutoPadFromString(FStringView StringVal)
     }
 }
 
-template<typename T>
-static bool IsEqualOrBroadcastable(TConstArrayView<T> ShapeA, TConstArrayView<T> ShapeB)
+static bool CheckElementwiseTensor(ENNETensorDataType DataType, const NNE::FSymbolicTensorShape& TensorShape)
 {
-    if(ShapeA.Num() < ShapeB.Num())
+    if (DataType != ENNETensorDataType::Float)
     {
+        UE_LOG(LogNNE, Warning, TEXT("Invalid DML tensor data type"));
         return false;
     }
 
-    int32 IdxOffset = 0;
+    if(!TensorShape.IsConcrete())
+    {
+        UE_LOG(LogNNE, Warning, TEXT("DML tensor shape must be concrete"));
+        return false;
+    }
 
-    if(ShapeB.Num() < ShapeA.Num())
+    if (NNE::FTensorShape::MakeFromSymbolic(TensorShape).Volume() == 0)
     {
-        IdxOffset = ShapeA.Num() - ShapeB.Num();
+        UE_LOG(LogNNE, Warning, TEXT("Invalid DML tensor size, it's 0"));
+        return false;
     }
 
-    for(int32 Idx = 0; Idx < ShapeB.Num(); ++Idx)
+    const int32 MinTensorRank(0), MaxTensorRank(DML_TENSOR_DIMENSION_COUNT_MAX1);
+
+    if (TensorShape.Rank() < MinTensorRank || TensorShape.Rank() > MaxTensorRank)
     {
-        if(ShapeA[Idx + IdxOffset] != ShapeB[Idx])
-        {
-            if(ShapeB[Idx] != 1)
-            {
-                return false;
-            }
-        }
+        UE_LOG(LogNNE, Warning, TEXT("Invalid DML tensor rank: %d [%d,%d]"), TensorShape.Rank(), MinTensorRank, MaxTensorRank);
+        return false;
     }
 
     return true;
