Add GPU stats for RDG DirectML runtime

#rnx

diff --git a/Engine/Plugins/Experimental/NNERuntimeRDG/Source/NNERuntimeRDG/Private/Dml/NNEDmlModel.cpp b/Engine/Plugins/Experimental/NNERuntimeRDG/Source/NNERuntimeRDG/Private/Dml/NNEDmlModel.cpp
--- a/Engine/Plugins/Experimental/NNERuntimeRDG/Source/NNERuntimeRDG/Private/Dml/NNEDmlModel.cpp
+++ b/Engine/Plugins/Experimental/NNERuntimeRDG/Source/NNERuntimeRDG/Private/Dml/NNEDmlModel.cpp
@@ -14,6 +14,19 @@
 namespace UE::NNERuntimeRDG::Private::Dml
 {
 
+#if STATS
+DECLARE_STATS_GROUP(TEXT("NNEDmlModel"), STATGROUP_NNEDmlModel, STATCAT_Advanced);
+DECLARE_MEMORY_STAT(TEXT("MemSizeWeights GPU"), STAT_MemSizeWeights, STATGROUP_NNEDmlModel);
+DECLARE_MEMORY_STAT(TEXT("MemSizePersistent GPU"), STAT_MemSizePersist, STATGROUP_NNEDmlModel);
+DECLARE_MEMORY_STAT(TEXT("MemSizeTemp GPU"), STAT_MemSizeTemp, STATGROUP_NNEDmlModel);
+DECLARE_DWORD_COUNTER_STAT(TEXT("InferenceCount"), STAT_InferenceCount, STATGROUP_NNEDmlModel);
+#endif
+
+#if HAS_GPU_STATS
+DECLARE_GPU_STAT_NAMED(GPU_STAT_DispatchTime, TEXT("NNEDmlModelDispatchTime"));
+DECLARE_GPU_STAT_NAMED(GPU_STAT_DispatchD3DTime, TEXT("NNEDmlModelDispatchD3DTime"));
+#endif
+
 FModelInfo* FModelInfo::Get()
 {
 	static FModelInfo Inst;
@@ -299,7 +312,7 @@ private:
 	TArray<DML_BUFFER_BINDING, TInlineAllocator<MaxNumOutputs>>		OutputBinds;
 	TArray<DML_BINDING_DESC, TInlineAllocator<MaxNumOutputs>>		OutputBindDescs;
 	ID3D12DynamicRHI*												DynamicRHI;
-	FModelInstance*															Model;
+	FModelInstance*													Model;
 };
 
 //
@@ -831,6 +844,17 @@ FModelInstance::FModelInstance()
 //
 FModelInstance::~FModelInstance()
 {
+
+#if STATS
+	DEC_MEMORY_STAT_BY(STAT_MemSizeWeights, MemSizeWeights);
+	DEC_MEMORY_STAT_BY(STAT_MemSizeTemp, MemSizeTemp);
+	DEC_MEMORY_STAT_BY(STAT_MemSizePersist, MemSizePersist);
+
+//	if (CompiledOp)
+//	{
+//		DEC_DWORD_STAT(STAT_InferenceCount);
+//	}
+#endif
 }
 
 //
@@ -925,7 +949,7 @@ bool FModelInstance::Init(TConstArrayView<uint8> ModelData, FDmlDeviceContext* I
 
 		if (!OpDesc.Op)
 		{
-			UE_LOG(LogNNE, Warning, TEXT("Error:Failed to create operator:%s"), *TypeName);
+			UE_LOG(LogNNE, Warning, TEXT("Error:Failed to create DML operator:%s"), *TypeName);
 			return false;
 		}
 
@@ -1002,7 +1026,7 @@ bool FModelInstance::InitCompiledOp(TConstArrayView<int32> OpInputIndices, uint6
 	Res = OpInit->Reset(UE_ARRAY_COUNT(CompiledOps), CompiledOps);
 	if (FAILED(Res))
 	{
-		UE_LOG(LogNNE, Warning, TEXT("Error:Failed to reset DirectML operator initializer"));
+		UE_LOG(LogNNE, Warning, TEXT("Error:Failed to reset DML operator initializer"));
 		return false;
 	}
 
@@ -1032,12 +1056,13 @@ bool FModelInstance::InitCompiledOp(TConstArrayView<int32> OpInputIndices, uint6
 		return false;
 	}
 
+	MemSizeWeights = TensorDataSize;
 	MemSizeTemp = ExecBindProps.TemporaryResourceSize;
 	MemSizePersist = ExecBindProps.PersistentResourceSize;
 
 	FEvent* Signal = FGenericPlatformProcess::GetSynchEventFromPool(false);
 
-	ENQUEUE_RENDER_COMMAND(FModel_SetTensorData)
+	ENQUEUE_RENDER_COMMAND(FDmlModelInstance_SetTensorData)
 	(
 		[
 			this, 
@@ -1062,8 +1087,9 @@ bool FModelInstance::InitCompiledOp(TConstArrayView<int32> OpInputIndices, uint6
 
 			if (TensorDataSize)
 			{
-				UploadFence = RHICreateGPUFence(TEXT("FInferenceModel_UploadFence"));
-				FRHIBuffer*		UploadBuff = CreateRHIBuffer(RHICmdList, TensorDataSize, BUF_ShaderResource | BUF_Dynamic | BUF_FastVRAM, ERHIAccess::CopySrc, TEXT("FInferenceModel_UploadBuffer"));
+				UploadFence = RHICreateGPUFence(TEXT("FDmlModelInstance_UploadFence"));
+
+				FBufferRHIRef	UploadBuff = CreateRHIBuffer(RHICmdList, TensorDataSize, BUF_ShaderResource | BUF_Dynamic | BUF_FastVRAM, ERHIAccess::CopySrc, TEXT("FDmlModelInstance_UploadBuffer"));
 				uint8*			UploadBuffPtr = static_cast<uint8*>(RHICmdList.LockBuffer(UploadBuff, 0, TensorDataSize, RLM_WriteOnly_NoOverwrite));
 				uint64			UploadOffset = 0;
 
@@ -1097,6 +1123,8 @@ bool FModelInstance::InitCompiledOp(TConstArrayView<int32> OpInputIndices, uint6
 
 				RHICmdList.UnlockBuffer(UploadBuff);
 				RHICmdList.WriteGPUFence(UploadFence);
+
+				INC_MEMORY_STAT_BY(STAT_MemSizeWeights, TensorDataSize);
 			}
 
 			if (MemSizePersist)
@@ -1104,8 +1132,9 @@ bool FModelInstance::InitCompiledOp(TConstArrayView<int32> OpInputIndices, uint6
 #ifdef NNE_USE_D3D12_RESOURCES
 				PersistBuff = CreateD3D12Buffer(MemSizePersist);
 #else		
-				PersistBuff = CreateRHIBuffer(RHICmdList, MemSizePersist, PersistBuffFlags, PersistBuffAccess, TEXT("FModelDml_PeristBuff"));
+				PersistBuff = CreateRHIBuffer(RHICmdList, MemSizePersist, PersistBuffFlags, PersistBuffAccess, TEXT("FDmlModelInstance_PeristBuff"));
 #endif
+				INC_MEMORY_STAT_BY(STAT_MemSizePersist, MemSizePersist);
 			}
 
 			if (MemSizeTemp)
@@ -1113,8 +1142,9 @@ bool FModelInstance::InitCompiledOp(TConstArrayView<int32> OpInputIndices, uint6
 #ifdef NNE_USE_D3D12_RESOURCES
 				TempBuff = CreateD3D12Buffer(MemSizeTemp);
 #else
-				TempBuff = CreateRHIBuffer(RHICmdList, MemSizeTemp, TempBuffFlags, TempBuffAccess, TEXT("FModelDml_TempBuff"));
+				TempBuff = CreateRHIBuffer(RHICmdList, MemSizeTemp, TempBuffFlags, TempBuffAccess, TEXT("FDmlModelInstance_TempBuff"));
 #endif
+				INC_MEMORY_STAT_BY(STAT_MemSizeTemp, MemSizeTemp);
 			}
 
 #ifdef NNE_USE_D3D12_RESOURCES
@@ -1128,7 +1158,7 @@ bool FModelInstance::InitCompiledOp(TConstArrayView<int32> OpInputIndices, uint6
 #ifdef NNE_USE_D3D12_RESOURCES
 				InitTempBuff = CreateD3D12Buffer(InitTempMemSize);
 #else
-				TempBuff = CreateRHIBuffer(RHICmdList, InitTempMemSize, TempBuffFlags, TempBuffAccess, TEXT("FModelDml_InitTempBuff"));
+				TempBuff = CreateRHIBuffer(RHICmdList, InitTempMemSize, TempBuffFlags, TempBuffAccess, TEXT("FDmlModelInstance_InitTempBuff"));
 #endif
 			}
 
@@ -1180,6 +1210,10 @@ END_SHADER_PARAMETER_STRUCT()
 //
 void FModelInstance::AddDispatchOps_RenderThread(FRDGBuilder& GraphBuilder)
 {
+#if STATS
+	INC_DWORD_STAT(STAT_InferenceCount);
+#endif
+
 	FDmlModelDispatchPassParameters* DispatchParams = GraphBuilder.AllocParameters<FDmlModelDispatchPassParameters>();
 
 	for (int32 Idx = 0; Idx < InputTensorIndices.Num(); ++Idx)
@@ -1202,8 +1236,12 @@ void FModelInstance::AddDispatchOps_RenderThread(FRDGBuilder& GraphBuilder)
 		DispatchParams->OutputBuffers.Emplace(AllTensorRDGRefs[OutputTensorIndices[Idx]]->GetBuffer(), ERHIAccess::UAVCompute);
 	}
 
+#if HAS_GPU_STATS
+	RDG_GPU_STAT_SCOPE(GraphBuilder, GPU_STAT_DispatchTime);
+#endif
+
 	GraphBuilder.AddPass(
-		RDG_EVENT_NAME("FDmlModelDispatch"),
+		RDG_EVENT_NAME("FDmlModelInstance_Dispatch"),
 		DispatchParams,
 		ERDGPassFlags::Compute | ERDGPassFlags::NeverCull,
 		[this, DispatchParams, NumWeightTensors](FRHICommandListImmediate& RHICmdList)
@@ -1228,70 +1266,75 @@ void FModelInstance::AddDispatchOps_RenderThread(FRDGBuilder& GraphBuilder)
 				RHIOutputBuffers.Add(RDGBuffer->GetRHI());
 			}
 
-			RHICmdList.EnqueueLambda(
-				[this, InputBuffers = MoveTemp(RHIInputBuffers), OutputBuffers = MoveTemp(RHIOutputBuffers)](FRHICommandListImmediate& RHICmdList)
-				{
-					TArray<CD3DX12_RESOURCE_BARRIER, TInlineAllocator<MaxNumInputs + MaxNumOutputs>>	PreBarriers;
-					TArray<CD3DX12_RESOURCE_BARRIER, TInlineAllocator<MaxNumOutputs * 2>>				PostBarriers;
-
-					for (FRHIBuffer* Buffer : InputBuffers)
+			{
+#if HAS_GPU_STATS
+				SCOPED_GPU_STAT(RHICmdList, GPU_STAT_DispatchD3DTime)
+#endif
+				RHICmdList.EnqueueLambda(
+					[this, InputBuffers = MoveTemp(RHIInputBuffers), OutputBuffers = MoveTemp(RHIOutputBuffers)](FRHICommandListImmediate& RHICmdList)
 					{
-						if (!Buffer)
-						{
-							continue;
-						}
+						TArray<CD3DX12_RESOURCE_BARRIER, TInlineAllocator<MaxNumInputs + MaxNumOutputs>>	PreBarriers;
+						TArray<CD3DX12_RESOURCE_BARRIER, TInlineAllocator<MaxNumOutputs * 2>>				PostBarriers;
 
-						ID3D12Resource* Resource = DynamicRHI->RHIGetResource(Buffer);
+						for (FRHIBuffer* Buffer : InputBuffers)
+						{
+							if (!Buffer)
+							{
+								continue;
+							}
 
-						//Note: We should not assume COPY_DEST state
-						PreBarriers.Emplace(
-							CD3DX12_RESOURCE_BARRIER::Transition(
-								Resource,
-								D3D12_RESOURCE_STATE_COPY_DEST,
-								D3D12_RESOURCE_STATE_UNORDERED_ACCESS)
-						);
-					}
+							ID3D12Resource* Resource = DynamicRHI->RHIGetResource(Buffer);
 
-					for (FRHIBuffer* Buffer : OutputBuffers)
-					{
-						ID3D12Resource* Resource = DynamicRHI->RHIGetResource(Buffer);
-
-						if ((Buffer->GetUsage() & EBufferUsageFlags::SourceCopy) == EBufferUsageFlags::SourceCopy)
-						{
+							//Note: We should not assume COPY_DEST state
 							PreBarriers.Emplace(
 								CD3DX12_RESOURCE_BARRIER::Transition(
 									Resource,
-									D3D12_RESOURCE_STATE_COMMON,
+									D3D12_RESOURCE_STATE_COPY_DEST,
 									D3D12_RESOURCE_STATE_UNORDERED_ACCESS)
 							);
-
-							PostBarriers.Emplace(
-								CD3DX12_RESOURCE_BARRIER::Transition(
-									Resource,
-									D3D12_RESOURCE_STATE_UNORDERED_ACCESS,
-									D3D12_RESOURCE_STATE_COPY_SOURCE)
-							);
 						}
 
-						PostBarriers.Add(CD3DX12_RESOURCE_BARRIER::UAV(Resource));
-					}
+						for (FRHIBuffer* Buffer : OutputBuffers)
+						{
+							ID3D12Resource* Resource = DynamicRHI->RHIGetResource(Buffer);
+
+							if ((Buffer->GetUsage() & EBufferUsageFlags::SourceCopy) == EBufferUsageFlags::SourceCopy)
+							{
+								PreBarriers.Emplace(
+									CD3DX12_RESOURCE_BARRIER::Transition(
+										Resource,
+										D3D12_RESOURCE_STATE_COMMON,
+										D3D12_RESOURCE_STATE_UNORDERED_ACCESS)
+								);
+
+								PostBarriers.Emplace(
+									CD3DX12_RESOURCE_BARRIER::Transition(
+										Resource,
+										D3D12_RESOURCE_STATE_UNORDERED_ACCESS,
+										D3D12_RESOURCE_STATE_COPY_SOURCE)
+								);
+							}
+
+							PostBarriers.Add(CD3DX12_RESOURCE_BARRIER::UAV(Resource));
+						}
 
-					//Note: We should use this instead of NNE_USE_D3D12_RESOURCES
-					//FBufferRHIRef TempBuff = MemSizeTemp ? CreateRHIBuffer(RHICmdList, MemSizeTemp, TempBuffUsage, TempBuffAccess, TEXT("FModel_Dispatch_TempBuff")) : nullptr;
+						//Note: We should use this instead of NNE_USE_D3D12_RESOURCES
+						//FBufferRHIRef TempBuff = MemSizeTemp ? CreateRHIBuffer(RHICmdList, MemSizeTemp, TempBuffUsage, TempBuffAccess, TEXT("FDmlModelInstance_Dispatch_TempBuff")) : nullptr;
 
-					BindingTable->Bind(CompiledOp, InputBuffers, OutputBuffers, PersistBuff, TempBuff);
+						BindingTable->Bind(CompiledOp, InputBuffers, OutputBuffers, PersistBuff, TempBuff);
 
-					ID3D12GraphicsCommandList* D3DCmdList = nullptr;
+						ID3D12GraphicsCommandList* D3DCmdList = nullptr;
 
-					D3DCmdList = DynamicRHI->RHIGetGraphicsCommandList(DevCtx->DeviceIndex);
-					D3DCmdList->SetDescriptorHeaps(1, &DescHeap);
-					D3DCmdList->ResourceBarrier(PreBarriers.Num(), PreBarriers.GetData());
-					DevCtx->CmdRec->RecordDispatch(D3DCmdList, CompiledOp, BindingTable->Get());
-					D3DCmdList->ResourceBarrier(PostBarriers.Num(), PostBarriers.GetData());
+						D3DCmdList = DynamicRHI->RHIGetGraphicsCommandList(DevCtx->DeviceIndex);
+						D3DCmdList->SetDescriptorHeaps(1, &DescHeap);
+						D3DCmdList->ResourceBarrier(PreBarriers.Num(), PreBarriers.GetData());
+						DevCtx->CmdRec->RecordDispatch(D3DCmdList, CompiledOp, BindingTable->Get());
+						D3DCmdList->ResourceBarrier(PostBarriers.Num(), PostBarriers.GetData());
 
-					DynamicRHI->RHIFinishExternalComputeWork(DevCtx->DeviceIndex, D3DCmdList);
-				}
-			);
+						DynamicRHI->RHIFinishExternalComputeWork(DevCtx->DeviceIndex, D3DCmdList);
+					}
+				);
+			}
 		}
 	);
 }
@@ -1305,7 +1348,7 @@ FOperatorDml* FModelInstance::OpCreate(const FString& OpName, TArrayView<const N
 
 	if (!CreateFn)
 	{
-		UE_LOG(LogNNE, Warning, TEXT("Dml MLOperatorRegistry failed to find operator:%s"), *OpName);
+		UE_LOG(LogNNE, Warning, TEXT("Error:failed to find DML operator:%s"), *OpName);
 		return nullptr;
 	}
 
@@ -1315,7 +1358,7 @@ FOperatorDml* FModelInstance::OpCreate(const FString& OpName, TArrayView<const N
 	{
 		delete Op;
 
-		UE_LOG(LogNNE, Warning, TEXT("Error:Failed to initialize operator:%s"), *OpName);
+		UE_LOG(LogNNE, Warning, TEXT("Error:Failed to initialize DML operator:%s"), *OpName);
 		return nullptr;
 	}
 
@@ -1357,7 +1400,7 @@ ID3D12Resource* FModelInstance::CreateD3D12Buffer(uint32 Size, D3D12_RESOURCE_ST
 
 	if (FAILED(Res))
 	{
-		UE_LOG(LogNNE, Warning, TEXT("Error:FInferenceModel failed to create D3D12 resource"));
+		UE_LOG(LogNNE, Warning, TEXT("Error:FDmlModelInstance failed to create D3D12 resource"));
 		return nullptr;
 	}
 
@@ -1378,7 +1421,7 @@ int FModelInstance::PrepareTensorShapesAndData()
 	{
 		if (!SymbolicTensorDesc.GetShape().IsConcrete())
 		{
-			UE_LOG(LogNNE, Warning, TEXT("DML runtime does not support model with variable shapes yet."));
+			UE_LOG(LogNNE, Warning, TEXT("FDmlModelInstance does not support model with variable shapes yet."));
 			return -1;
 		}
 	}
diff --git a/Engine/Plugins/Experimental/NNERuntimeRDG/Source/NNERuntimeRDG/Private/Dml/NNEDmlModel.h b/Engine/Plugins/Experimental/NNERuntimeRDG/Source/NNERuntimeRDG/Private/Dml/NNEDmlModel.h
--- a/Engine/Plugins/Experimental/NNERuntimeRDG/Source/NNERuntimeRDG/Private/Dml/NNEDmlModel.h
+++ b/Engine/Plugins/Experimental/NNERuntimeRDG/Source/NNERuntimeRDG/Private/Dml/NNEDmlModel.h
@@ -92,6 +92,7 @@ private:
 #else
 	FBufferRHIRef						PersistBuff;
 #endif
+	uint64								MemSizeWeights;
 	uint64								MemSizeTemp;
 	uint64								MemSizePersist;
 	ID3D12DynamicRHI*					DynamicRHI;
