Use RHI buffers instead of D3D12Resource

#rnx

diff --git a/Engine/Plugins/Experimental/NNERuntimeRDG/Source/NNERuntimeRDG/Private/Dml/NNEDmlModel.cpp b/Engine/Plugins/Experimental/NNERuntimeRDG/Source/NNERuntimeRDG/Private/Dml/NNEDmlModel.cpp
--- a/Engine/Plugins/Experimental/NNERuntimeRDG/Source/NNERuntimeRDG/Private/Dml/NNEDmlModel.cpp
+++ b/Engine/Plugins/Experimental/NNERuntimeRDG/Source/NNERuntimeRDG/Private/Dml/NNEDmlModel.cpp
@@ -135,11 +135,7 @@ public:
 		return true;
 	}
 
-#ifdef NNE_USE_D3D12_RESOURCES
-	void Bind(IDMLOperatorInitializer* InOpInit, TConstArrayView<FRHIBuffer*> InputBuffers, ID3D12Resource* InPersistBuff, ID3D12Resource* InTempBuff = nullptr)
-#else
 	void Bind(IDMLOperatorInitializer* InOpInit, TConstArrayView<FRHIBuffer*> InputBuffers, FRHIBuffer* InPersistBuff, FRHIBuffer* InTempBuff = nullptr)
-#endif
 	{
 		Reset(InOpInit);
 
@@ -167,11 +163,7 @@ public:
 
 		if (InPersistBuff)
 		{
-#ifdef NNE_USE_D3D12_RESOURCES
-			PersistBind = DML_BUFFER_BINDING { InPersistBuff, 0, InPersistBuff->GetDesc().Width };
-#else
 			PersistBind = MakeBind(InPersistBuff);
-#endif
 		}
 
 		BindingTable->BindOutputs(1, &PersistBindDesc);
@@ -181,21 +173,13 @@ public:
 
 		if (InTempBuff)
 		{
-#ifdef NNE_USE_D3D12_RESOURCES
-			TempBind = { InTempBuff, 0, InTempBuff->GetDesc().Width };
-#else
 			TempBind = MakeBind(InTempBuff);
-#endif
+
 			BindingTable->BindTemporaryResource(&TempBindDesc);
 		}
 	}
 
-#ifdef NNE_USE_D3D12_RESOURCES
-	void Bind(IDMLCompiledOperator* Op, TConstArrayView<FRHIBuffer*> InputBuffers, TConstArrayView<FRHIBuffer*> OutputBuffers, ID3D12Resource* InPersistBuff = nullptr, ID3D12Resource* InTempBuff = nullptr)
-#else
-	
 	void Bind(IDMLCompiledOperator* Op, TConstArrayView<FRHIBuffer*> InputBuffers, TConstArrayView<FRHIBuffer*> OutputBuffers, FRHIBuffer* InPersistBuff = nullptr, FRHIBuffer* InTempBuff = nullptr)
-#endif
 	{
 		Reset(Op);
 
@@ -217,11 +201,8 @@ public:
 
 		if (InPersistBuff)
 		{
-#ifdef NNE_USE_D3D12_RESOURCES
-			PersistBind = { InPersistBuff, 0, InPersistBuff->GetDesc().Width };
-#else
 			PersistBind =  MakeBind(InPersistBuff);MakeBind(InPersistBuff);
-#endif
+
 			BindingTable->BindPersistentResource(&PersistBindDesc);
 		}
 
@@ -230,17 +211,12 @@ public:
 
 		if (InTempBuff)
 		{
-			
-#ifdef NNE_USE_D3D12_RESOURCES
-			TempBind = { InTempBuff, 0, InTempBuff->GetDesc().Width };
-#else
 			TempBind = MakeBind(InTempBuff);
-#endif
+
 			BindingTable->BindTemporaryResource(&TempBindDesc);
 		}
 	}
 
-
 	IDMLBindingTable* Get()
 	{
 		return BindingTable;
@@ -1016,7 +992,7 @@ bool FModelInstance::InitCompiledOp(TConstArrayView<int32> OpInputIndices, uint6
 	static constexpr EBufferUsageFlags	PersistBuffFlags = BUF_Static | BUF_ShaderResource | BUF_UnorderedAccess;
 	static constexpr ERHIAccess			PersistBuffAccess = ERHIAccess::UAVMask;
 
-	static constexpr EBufferUsageFlags	TempBuffFlags = BUF_Volatile | BUF_UnorderedAccess;
+	static constexpr EBufferUsageFlags	TempBuffFlags = BUF_Static | BUF_ShaderResource | BUF_UnorderedAccess;
 	static constexpr ERHIAccess			TempBuffAccess = ERHIAccess::UAVMask;
 
 	HRESULT					Res;
@@ -1129,37 +1105,21 @@ bool FModelInstance::InitCompiledOp(TConstArrayView<int32> OpInputIndices, uint6
 
 			if (MemSizePersist)
 			{
-#ifdef NNE_USE_D3D12_RESOURCES
-				PersistBuff = CreateD3D12Buffer(MemSizePersist);
-#else		
 				PersistBuff = CreateRHIBuffer(RHICmdList, MemSizePersist, PersistBuffFlags, PersistBuffAccess, TEXT("FDmlModelInstance_PeristBuff"));
-#endif
 				INC_MEMORY_STAT_BY(STAT_MemSizePersist, MemSizePersist);
 			}
 
 			if (MemSizeTemp)
 			{
-#ifdef NNE_USE_D3D12_RESOURCES
-				TempBuff = CreateD3D12Buffer(MemSizeTemp);
-#else
 				TempBuff = CreateRHIBuffer(RHICmdList, MemSizeTemp, TempBuffFlags, TempBuffAccess, TEXT("FDmlModelInstance_TempBuff"));
-#endif
 				INC_MEMORY_STAT_BY(STAT_MemSizeTemp, MemSizeTemp);
 			}
 
-#ifdef NNE_USE_D3D12_RESOURCES
-			ID3D12Resource* InitTempBuff = nullptr;
-#else
 			FBufferRHIRef InitTempBuff;
-#endif
 
 			if (InitTempMemSize)
 			{
-#ifdef NNE_USE_D3D12_RESOURCES
-				InitTempBuff = CreateD3D12Buffer(InitTempMemSize);
-#else
-				TempBuff = CreateRHIBuffer(RHICmdList, InitTempMemSize, TempBuffFlags, TempBuffAccess, TEXT("FDmlModelInstance_InitTempBuff"));
-#endif
+				InitTempBuff = CreateRHIBuffer(RHICmdList, InitTempMemSize, TempBuffFlags, TempBuffAccess, TEXT("FDmlModelInstance_InitTempBuff"));
 			}
 
 			RHICmdList.EnqueueLambda(
@@ -1318,9 +1278,6 @@ void FModelInstance::AddDispatchOps_RenderThread(FRDGBuilder& GraphBuilder)
 							PostBarriers.Add(CD3DX12_RESOURCE_BARRIER::UAV(Resource));
 						}
 
-						//Note: We should use this instead of NNE_USE_D3D12_RESOURCES
-						//FBufferRHIRef TempBuff = MemSizeTemp ? CreateRHIBuffer(RHICmdList, MemSizeTemp, TempBuffUsage, TempBuffAccess, TEXT("FDmlModelInstance_Dispatch_TempBuff")) : nullptr;
-
 						BindingTable->Bind(CompiledOp, InputBuffers, OutputBuffers, PersistBuff, TempBuff);
 
 						ID3D12GraphicsCommandList* D3DCmdList = nullptr;
diff --git a/Engine/Plugins/Experimental/NNERuntimeRDG/Source/NNERuntimeRDG/Private/Dml/NNEDmlModel.h b/Engine/Plugins/Experimental/NNERuntimeRDG/Source/NNERuntimeRDG/Private/Dml/NNEDmlModel.h
--- a/Engine/Plugins/Experimental/NNERuntimeRDG/Source/NNERuntimeRDG/Private/Dml/NNEDmlModel.h
+++ b/Engine/Plugins/Experimental/NNERuntimeRDG/Source/NNERuntimeRDG/Private/Dml/NNEDmlModel.h
@@ -5,8 +5,6 @@
 #include "NNEDmlCommon.h"
 #include "NNERuntimeRDGModel.h"
 
-#define NNE_USE_D3D12_RESOURCES
-
 struct ID3D12DynamicRHI;
 
 namespace UE::NNERuntimeRDG::Private::Dml
@@ -86,12 +84,8 @@ private:
 
 	TArray<int32>						ConstantCPUTensorIndices;
 
-#ifdef NNE_USE_D3D12_RESOURCES
-	TComPtr<ID3D12Resource>				PersistBuff;
-	TComPtr<ID3D12Resource>				TempBuff;
-#else
 	FBufferRHIRef						PersistBuff;
-#endif
+	FBufferRHIRef						TempBuff;
 	uint64								MemSizeWeights;
 	uint64								MemSizeTemp;
 	uint64								MemSizePersist;
