[NNE] Add validation to RDG Dml BatchNormalization

#rnx

diff --git a/Engine/Plugins/Experimental/NNERuntimeRDG/Source/NNERuntimeRDG/Private/Dml/NNEDmlOperatorBatchNormalization.cpp b/Engine/Plugins/Experimental/NNERuntimeRDG/Source/NNERuntimeRDG/Private/Dml/NNEDmlOperatorBatchNormalization.cpp
--- a/Engine/Plugins/Experimental/NNERuntimeRDG/Source/NNERuntimeRDG/Private/Dml/NNEDmlOperatorBatchNormalization.cpp
+++ b/Engine/Plugins/Experimental/NNERuntimeRDG/Source/NNERuntimeRDG/Private/Dml/NNEDmlOperatorBatchNormalization.cpp
@@ -2,6 +2,7 @@
 
 #ifdef NNE_USE_DIRECTML
 #include "NNEDmlOperator.h"
+#include "NNEDmlOperatorUtils.h"
 
 namespace UE::NNERuntimeRDG::Private::Dml
 {
@@ -29,7 +30,37 @@ public:
 
 	static bool Validate(const NNE::FAttributeMap& AttributeMap, TConstArrayView<ENNETensorDataType> InputTypes, TConstArrayView<NNE::FSymbolicTensorShape> InputShapes)
 	{
-		//TODO
+		if(InputShapes.Num() != Count)
+		{
+			UE_LOG(LogNNE, Warning, TEXT("DML input tensors must be %d"), Count);
+        	return false;
+		}
+		const int32 InputRank = InputShapes[0].Rank();
+		if(InputRank > 8)
+		{
+			UE_LOG(LogNNE, Warning, TEXT("DML InputTensor rank should be between 1 and 8, it's %d"), InputRank);
+        	return false;
+		}
+		const int32 bTrainingMode = AttributeMap.GetValueOrDefault<int32>(TEXT("training_mode"), 0);
+		if (bTrainingMode)
+		{
+			UE_LOG(LogNNE, Warning, TEXT("DML BatchNormalization doesn't support training mode"));
+			return false;
+		}
+
+		for (int32 Idx = X; Idx < Count; ++Idx)
+		{
+			if(!CheckElementwiseTensor(InputTypes[Idx], InputShapes[Idx]))
+			{
+				return false;
+			}
+			if(!IsEqualOrBroadcastable(InputShapes[0].GetData(), InputShapes[Idx].GetData()))
+			{
+				UE_LOG(LogNNE, Warning, TEXT("DML BatchNormalization other tensors' shapes must be equal or broadcastable to input tensor's"));
+				return false;
+			}
+		}
+
 		return true;
 	}
 
diff --git a/Engine/Plugins/Experimental/NNERuntimeRDG/Source/NNERuntimeRDG/Private/Dml/NNEDmlOperatorUtils.h b/Engine/Plugins/Experimental/NNERuntimeRDG/Source/NNERuntimeRDG/Private/Dml/NNEDmlOperatorUtils.h
--- a/Engine/Plugins/Experimental/NNERuntimeRDG/Source/NNERuntimeRDG/Private/Dml/NNEDmlOperatorUtils.h
+++ b/Engine/Plugins/Experimental/NNERuntimeRDG/Source/NNERuntimeRDG/Private/Dml/NNEDmlOperatorUtils.h
@@ -39,6 +39,43 @@ static EAutoPad AutoPadFromString(FStringView StringVal)
     }
 }
 
+template<typename T>
+static bool IsEqualOrBroadcastable(TConstArrayView<T> ShapeA, TConstArrayView<T> ShapeB)
+{
+    if(ShapeA.Num() < ShapeB.Num())
+    {
+        return false;
+    }
+
+    int32 BIdx = 0;
+
+    for(int32 Idx = 0; Idx < ShapeA.Num() && BIdx < ShapeB.Num(); ++Idx)
+    {
+        if(BIdx != 0)
+        {
+            if(ShapeA[Idx] != ShapeB[BIdx])
+            {
+                if(ShapeB[BIdx] != 1)
+                {
+                    return false;
+                }
+            }
+            ++BIdx;
+        }
+        if(BIdx == 0 && ShapeA[Idx] == ShapeB[BIdx])
+        {
+            ++BIdx;
+        }
+    }
+
+    if(BIdx == 0)
+    {
+        return false;
+    }
+
+    return true;
+}
+
 static bool CheckElementwiseTensor(ENNETensorDataType DataType, const NNE::FSymbolicTensorShape& TensorShape)
 {
     if (DataType != ENNETensorDataType::Float)
