[NNE] Add validation to Dml cast

#rnx

diff --git a/Engine/Plugins/Experimental/NNERuntimeRDG/Source/NNERuntimeRDG/Private/Dml/NNEDmlOperatorCast.cpp b/Engine/Plugins/Experimental/NNERuntimeRDG/Source/NNERuntimeRDG/Private/Dml/NNEDmlOperatorCast.cpp
--- a/Engine/Plugins/Experimental/NNERuntimeRDG/Source/NNERuntimeRDG/Private/Dml/NNEDmlOperatorCast.cpp
+++ b/Engine/Plugins/Experimental/NNERuntimeRDG/Source/NNERuntimeRDG/Private/Dml/NNEDmlOperatorCast.cpp
@@ -19,7 +19,17 @@ public:
 
 	static bool Validate(const NNE::FAttributeMap& AttributeMap, TConstArrayView<ENNETensorDataType> InputTypes, TConstArrayView<NNE::FSymbolicTensorShape> InputShapes)
 	{
-		//TODO
+		if(InputShapes.Num() != 1)
+		{
+			UE_LOG(LogNNE, Warning, TEXT("There must be only 1 DML input tensor"));
+        	return false;
+		}
+
+		if(!CheckGenericTensor(InputTypes[0], InputShapes[0]))
+		{
+			return false;
+		}
+
 		return true;
 	}
 
diff --git a/Engine/Plugins/Experimental/NNERuntimeRDG/Source/NNERuntimeRDG/Private/Dml/NNEDmlOperatorUtils.h b/Engine/Plugins/Experimental/NNERuntimeRDG/Source/NNERuntimeRDG/Private/Dml/NNEDmlOperatorUtils.h
--- a/Engine/Plugins/Experimental/NNERuntimeRDG/Source/NNERuntimeRDG/Private/Dml/NNEDmlOperatorUtils.h
+++ b/Engine/Plugins/Experimental/NNERuntimeRDG/Source/NNERuntimeRDG/Private/Dml/NNEDmlOperatorUtils.h
@@ -76,14 +76,8 @@ static bool IsEqualOrBroadcastable(TConstArrayView<T> ShapeA, TConstArrayView<T>
     return true;
 }
 
-static bool CheckElementwiseTensor(ENNETensorDataType DataType, const NNE::FSymbolicTensorShape& TensorShape)
+static bool CheckGenericTensor(ENNETensorDataType DataType, const NNE::FSymbolicTensorShape& TensorShape)
 {
-    if (DataType != ENNETensorDataType::Float)
-    {
-        UE_LOG(LogNNE, Warning, TEXT("Invalid DML tensor data type"));
-        return false;
-    }
-
     if(!TensorShape.IsConcrete())
     {
         UE_LOG(LogNNE, Warning, TEXT("DML tensor shape must be concrete"));
@@ -107,6 +101,22 @@ static bool CheckElementwiseTensor(ENNETensorDataType DataType, const NNE::FSymb
     return true;
 }
 
+static bool CheckElementwiseTensor(ENNETensorDataType DataType, const NNE::FSymbolicTensorShape& TensorShape)
+{
+    if (DataType != ENNETensorDataType::Float)
+    {
+        UE_LOG(LogNNE, Warning, TEXT("Invalid DML tensor data type"));
+        return false;
+    }
+
+    if(!CheckGenericTensor(DataType, TensorShape))
+    {
+        return false;
+    }
+
+    return true;
+}
+
 static Util::FSmallUIntArray KernelPadding(
     TConstArrayView<uint32> InputShape, TConstArrayView<uint32> WindowSize, 
     TConstArrayView<uint32> Dilations, TConstArrayView<uint32> Strides
