[NNE] Add validation to Dml concat

#rnx

diff --git a/Engine/Plugins/Experimental/NNERuntimeRDG/Source/NNERuntimeRDG/Private/Dml/NNEDmlOperatorConcat.cpp b/Engine/Plugins/Experimental/NNERuntimeRDG/Source/NNERuntimeRDG/Private/Dml/NNEDmlOperatorConcat.cpp
--- a/Engine/Plugins/Experimental/NNERuntimeRDG/Source/NNERuntimeRDG/Private/Dml/NNEDmlOperatorConcat.cpp
+++ b/Engine/Plugins/Experimental/NNERuntimeRDG/Source/NNERuntimeRDG/Private/Dml/NNEDmlOperatorConcat.cpp
@@ -5,6 +5,7 @@
 #include "NNEDmlOperator.h"
 #include "Misc/EnumerateRange.h"
 #include "Algo/Find.h"
+#include "NNEDmlOperatorUtils.h"
 
 namespace UE::NNERuntimeRDG::Private::Dml
 {
@@ -25,7 +26,60 @@ public:
 
 	static bool Validate(const NNE::FAttributeMap& AttributeMap, TConstArrayView<ENNETensorDataType> InputTypes, TConstArrayView<NNE::FSymbolicTensorShape> InputShapes)
 	{
-		//TODO
+		if(InputShapes.Num() == 0)
+		{
+			UE_LOG(LogNNE, Warning, TEXT("DML no input tensors"));
+        	return false;
+		}
+
+		const int32 InputDims = InputShapes[0].Rank();
+		auto NormalizeAxis = [InputDims] (int32 Axis) -> uint32
+		{
+			if (Axis < 0)
+			{
+				Axis += InputDims;
+			}
+			return (uint32) Axis;
+		};
+
+		uint32 Axis = NormalizeAxis(AttributeMap.GetValue<int32>(TEXT("axis")));
+		if(Axis >= (uint32) InputDims)
+		{
+			UE_LOG(LogNNE, Warning, TEXT("DML invalid axis: %d"), Axis);
+        	return false;
+		}
+
+		if(!CheckGenericTensor(InputTypes[0], InputShapes[0]))
+		{
+			return false;
+		}
+		
+		for (int32 Idx = 1; Idx < InputShapes.Num(); ++Idx)
+		{
+			if(!CheckGenericTensor(InputTypes[Idx], InputShapes[Idx]))
+			{
+				return false;
+			}
+
+			if(InputShapes[Idx].Rank() != InputDims)
+			{
+				UE_LOG(LogNNE, Warning, TEXT("DML concat rank mismatch for tensor %d"), Idx);
+				return false;
+			}
+
+			for(int32 Dim = 0; Dim < InputShapes[Idx].Rank(); ++Dim)
+			{
+				if(Dim != Axis)
+				{
+					if(InputShapes[Idx].GetData()[Dim] != InputShapes[0].GetData()[Dim])
+					{
+						UE_LOG(LogNNE, Warning, TEXT("DML concat dimension mismatch for tensor %d"), Idx);
+						return false;
+					}
+				}
+			}
+		}
+
 		return true;
 	}
 
