Fix RDG DML Gather op issues required by wav2vec model

#rnx

diff --git a/Engine/Plugins/Experimental/NNERuntimeRDG/Source/NNERuntimeRDG/Private/Dml/NNEDmlOperatorGather.cpp b/Engine/Plugins/Experimental/NNERuntimeRDG/Source/NNERuntimeRDG/Private/Dml/NNEDmlOperatorGather.cpp
--- a/Engine/Plugins/Experimental/NNERuntimeRDG/Source/NNERuntimeRDG/Private/Dml/NNEDmlOperatorGather.cpp
+++ b/Engine/Plugins/Experimental/NNERuntimeRDG/Source/NNERuntimeRDG/Private/Dml/NNEDmlOperatorGather.cpp
@@ -43,9 +43,12 @@ public:
 			return false;
 		}
 
-		if (IndicesTensor.GetDataType() != ENNETensorDataType::UInt32 && IndicesTensor.GetDataType() != ENNETensorDataType::Int32)
+		ENNETensorDataType IndicesDataType = IndicesTensor.GetDataType();
+
+		if (IndicesDataType != ENNETensorDataType::UInt32 && IndicesDataType != ENNETensorDataType::Int32 &&
+			IndicesDataType != ENNETensorDataType::UInt64 && IndicesDataType != ENNETensorDataType::Int64)
 		{
-			UE_LOG(LogNNE, Warning, TEXT("DML only supports UINT32/INT32 for indices tensor"));
+			UE_LOG(LogNNE, Warning, TEXT("DML only supports UINT32/INT32/UINT64/INT64 for indices tensor"));
 			return false;
 		}
 
@@ -78,9 +81,9 @@ public:
 			OutputShape.Add(InputShape.GetData()[DataRankIdx]);
 		}
 
-		if (InputShape.Rank() != OutputRank)
+		if (OutputTensor.GetShape().Rank() != OutputRank)
 		{
-			UE_LOG(LogNNE, Warning, TEXT("Output tensor rank must match input tensor rank"));
+			UE_LOG(LogNNE, Warning, TEXT("Output tensor rank must match computed output tensor rank"));
 			return false;
 		}
 
@@ -107,7 +110,7 @@ public:
 
 		if (!DmlOutputTensorDesc
 				.SetFromTensor(OutputTensor)
-				.SetShape(OutputShape)
+				.SetShape(OutputShape, InputTensor.GetShape().Rank())
 				.Validate())
 		{
 			UE_LOG(LogNNE, Error, TEXT("Failed to initialize tensor(s) for DML inference"));
