Fix lambda function used in checkf preventing from validating bad checkf condition parameters

#rnx

diff --git a/Engine/Plugins/Experimental/NNERuntimeRDG/Source/NNERuntimeRDG/Private/NNERuntimeRDGBase.cpp b/Engine/Plugins/Experimental/NNERuntimeRDG/Source/NNERuntimeRDG/Private/NNERuntimeRDGBase.cpp
--- a/Engine/Plugins/Experimental/NNERuntimeRDG/Source/NNERuntimeRDG/Private/NNERuntimeRDGBase.cpp
+++ b/Engine/Plugins/Experimental/NNERuntimeRDG/Source/NNERuntimeRDG/Private/NNERuntimeRDGBase.cpp
@@ -94,15 +94,23 @@ void FInputValidator::AddRequired(int32 TemplateIdx)
 //
 void FAttributeValidator::AddOptional(const FString& Name, ENNEAttributeDataType Type)
 {
-	checkf(nullptr == OptionalAttributes.FindByPredicate([Name](const FEntry& Other) { return Other.Name == Name; }), TEXT("Attribute name should be unique"));
-	checkf(nullptr == RequiredAttributes.FindByPredicate([Name](const FEntry& Other) { return Other.Name == Name; }), TEXT("Attribute name should be unique"));
+#if DO_CHECK
+	const bool bIsInOptional = (nullptr == OptionalAttributes.FindByPredicate([Name](const FEntry& Other) { return Other.Name == Name; }));
+	const bool bIsInRequired = (nullptr == RequiredAttributes.FindByPredicate([Name](const FEntry& Other) { return Other.Name == Name; }));
+	checkf(!bIsInOptional, TEXT("Attribute name should be unique"));	
+	checkf(!bIsInRequired, TEXT("Attribute name should be unique"));
+#endif
 	OptionalAttributes.Emplace(Name, Type);
 }
 
 void FAttributeValidator::AddRequired(const FString& Name, ENNEAttributeDataType Type)
 {
-	checkf(nullptr == OptionalAttributes.FindByPredicate([Name](const FEntry& Other) { return Other.Name == Name; }), TEXT("Attribute name should be unique"));
-	checkf(nullptr == RequiredAttributes.FindByPredicate([Name](const FEntry& Other) { return Other.Name == Name; }), TEXT("Attribute name should be unique"));
+#if DO_CHECK
+	const bool bIsInOptional = (nullptr == OptionalAttributes.FindByPredicate([Name](const FEntry& Other) { return Other.Name == Name; }));
+	const bool bIsInRequired = (nullptr == RequiredAttributes.FindByPredicate([Name](const FEntry& Other) { return Other.Name == Name; }));
+	checkf(!bIsInOptional, TEXT("Attribute name should be unique"));
+	checkf(!bIsInRequired, TEXT("Attribute name should be unique"));
+#endif
 	RequiredAttributes.Emplace(Name, Type);
 }
 
