[NNE] Fix bug with inverted check

Causing CL: 27440959

#rnx

diff --git a/Engine/Plugins/Experimental/NNERuntimeRDG/Source/NNERuntimeRDG/Private/NNERuntimeRDGBase.cpp b/Engine/Plugins/Experimental/NNERuntimeRDG/Source/NNERuntimeRDG/Private/NNERuntimeRDGBase.cpp
--- a/Engine/Plugins/Experimental/NNERuntimeRDG/Source/NNERuntimeRDG/Private/NNERuntimeRDGBase.cpp
+++ b/Engine/Plugins/Experimental/NNERuntimeRDG/Source/NNERuntimeRDG/Private/NNERuntimeRDGBase.cpp
@@ -95,10 +95,10 @@ void FInputValidator::AddRequired(int32 TemplateIdx)
 void FAttributeValidator::AddOptional(const FString& Name, ENNEAttributeDataType Type)
 {
 #if DO_CHECK
-	const bool bIsInOptional = (nullptr == OptionalAttributes.FindByPredicate([Name](const FEntry& Other) { return Other.Name == Name; }));
-	const bool bIsInRequired = (nullptr == RequiredAttributes.FindByPredicate([Name](const FEntry& Other) { return Other.Name == Name; }));
-	checkf(!bIsInOptional, TEXT("Attribute name should be unique"));	
-	checkf(!bIsInRequired, TEXT("Attribute name should be unique"));
+	const bool bIsFoundInOptional = (nullptr != OptionalAttributes.FindByPredicate([Name](const FEntry& Other) { return Other.Name == Name; }));
+	const bool bIsFoundInRequired = (nullptr != RequiredAttributes.FindByPredicate([Name](const FEntry& Other) { return Other.Name == Name; }));
+	checkf(!bIsFoundInOptional, TEXT("Attribute name should be unique"));
+	checkf(!bIsFoundInRequired, TEXT("Attribute name should be unique"));
 #endif
 	OptionalAttributes.Emplace(Name, Type);
 }
@@ -106,10 +106,10 @@ void FAttributeValidator::AddOptional(const FString& Name, ENNEAttributeDataType
 void FAttributeValidator::AddRequired(const FString& Name, ENNEAttributeDataType Type)
 {
 #if DO_CHECK
-	const bool bIsInOptional = (nullptr == OptionalAttributes.FindByPredicate([Name](const FEntry& Other) { return Other.Name == Name; }));
-	const bool bIsInRequired = (nullptr == RequiredAttributes.FindByPredicate([Name](const FEntry& Other) { return Other.Name == Name; }));
-	checkf(!bIsInOptional, TEXT("Attribute name should be unique"));
-	checkf(!bIsInRequired, TEXT("Attribute name should be unique"));
+	const bool bIsFoundInOptional = (nullptr != OptionalAttributes.FindByPredicate([Name](const FEntry& Other) { return Other.Name == Name; }));
+	const bool bIsFoundInRequired = (nullptr != RequiredAttributes.FindByPredicate([Name](const FEntry& Other) { return Other.Name == Name; }));
+	checkf(!bIsFoundInOptional, TEXT("Attribute name should be unique"));
+	checkf(!bIsFoundInRequired, TEXT("Attribute name should be unique"));
 #endif
 	RequiredAttributes.Emplace(Name, Type);
 }
