Add RDG DML suport for varying input shapes

#rnx

diff --git a/Engine/Plugins/Experimental/NNERuntimeRDG/Source/NNERuntimeRDG/Private/Dml/NNEDmlModel.cpp b/Engine/Plugins/Experimental/NNERuntimeRDG/Source/NNERuntimeRDG/Private/Dml/NNEDmlModel.cpp
--- a/Engine/Plugins/Experimental/NNERuntimeRDG/Source/NNERuntimeRDG/Private/Dml/NNEDmlModel.cpp
+++ b/Engine/Plugins/Experimental/NNERuntimeRDG/Source/NNERuntimeRDG/Private/Dml/NNEDmlModel.cpp
@@ -75,50 +75,43 @@ bool FModelInfo::ValidateGuidAndVersion(const uint8* InGuid, const uint8* InVers
 //
 //
 //
-class FModelInstance::FDebugName
+inline FModelInstance::FDebugName::FDebugName()
 {
-	static constexpr int32 Size = 128;
-
-public:
-
-	FDebugName()
-	{
-		Str[0] = '\0';
-		Length = 0;
-	}
-
-	FDebugName(const FString& InStr)
-	{
-		FTCHARToUTF8 Conv(*InStr);
-
-		Length = Conv.Length() + 1 < Size ? Conv.Length() + 1 : Size - 1;
-		FCStringAnsi::Strncpy(Str, Conv.Get(), Length);
-		Str[Length] = '\0';
-	}
-
-	FDebugName(FStringView InStr)
-	{
-		FTCHARToUTF8 Conv(InStr.GetData());
+	Str[0] = '\0';
+	Length = 0;
+}
 
-		Length = Conv.Length() + 1 < Size ? Conv.Length() + 1 : Size - 1;
-		FCStringAnsi::Strncpy(Str, Conv.Get(), Length);
-		Str[Length] = '\0';
-	}
+//
+//
+//
+inline FModelInstance::FDebugName::FDebugName(const FString& InStr)
+{
+	FTCHARToUTF8 Conv(*InStr);
 
-	~FDebugName()
-	{
-	}
+	Length = Conv.Length() + 1 < Size ? Conv.Length() + 1 : Size - 1;
+	FCStringAnsi::Strncpy(Str, Conv.Get(), Length);
+	Str[Length] = '\0';
+}
 
-	const char* Get() const
-	{
-		return Str;
-	}
+//
+//
+//
+inline FModelInstance::FDebugName::FDebugName(FStringView InStr)
+{
+	FTCHARToUTF8 Conv(InStr.GetData());
 
-private:
+	Length = Conv.Length() + 1 < Size ? Conv.Length() + 1 : Size - 1;
+	FCStringAnsi::Strncpy(Str, Conv.Get(), Length);
+	Str[Length] = '\0';
+}
 
-	char	Str[Size];
-	int32	Length;
-};
+//
+//
+//
+inline const char* FModelInstance::FDebugName::Get() const
+{
+	return Str;
+}
 
 //
 //
@@ -127,6 +120,9 @@ class FModelInstance::FBindingTable
 {
 public:
 
+	//
+	//
+	//
 	bool Init(FModelInstance* InModel)
 	{
 		Model = InModel;
@@ -135,6 +131,9 @@ public:
 		return true;
 	}
 
+	//
+	//
+	//
 	void Bind(IDMLOperatorInitializer* InOpInit, TConstArrayView<FRHIBuffer*> InputBuffers, FRHIBuffer* InPersistBuff, FRHIBuffer* InTempBuff = nullptr)
 	{
 		Reset(InOpInit);
@@ -246,7 +245,7 @@ private:
 			Res = Model->DevCtx->Device->CreateBindingTable(&Desc, DML_PPV_ARGS(&BindingTable));
 			if (!BindingTable)
 			{
-				UE_LOG(LogNNE, Warning, TEXT("Failed to create DML binding table, res:%d"), Res);
+				UE_LOG(LogNNE, Error, TEXT("Failed to create DML binding table, res:%d"), Res);
 				return false;
 			}
 		}
@@ -296,32 +295,6 @@ private:
 //
 class FModelInstance::FGraphBuilder
 {
-public:
-
-	struct FOpDesc
-	{
-		FOperatorDml*	Op;
-		int32			InputStart;
-		int32			InputCount;
-		int32			OutputStart;
-		int32			OutputCount;
-		FDebugName		DbgName;
-	};
-
-	struct FGraphDesc
-	{
-		TConstArrayView<NNE::Internal::FTensor>	AllTensors;
-		TConstArrayView<int32>		InputIndices;
-		TConstArrayView<int32>		OutputIndices;
-		TConstArrayView<int32>		WeightIndices;
-		TConstArrayView<int32>		ConstantCPUIndices;
-		TConstArrayView<int32>		IntermediateIndices;
-		TConstArrayView<FTensorRDG>	WeightTensors;
-		TConstArrayView<FOpDesc>	Operators;
-		TConstArrayView<int32>		OpInputIndices;
-		TConstArrayView<int32>		OpOutputIndices;
-	};
-	
 private:
 
 	enum class EEdgeType
@@ -331,6 +304,9 @@ private:
 		Intermediate
 	};
 
+	//
+	//
+	//
 	struct FEdge
 	{
 		EEdgeType	Type;
@@ -378,19 +354,23 @@ private:
 
 public:
 
-	IDMLCompiledOperator* Compile(FDmlDeviceContext* InDevCtx, const FGraphDesc& InGraph)
+	//
+	//
+	//
+	IDMLCompiledOperator* Compile(const FModelInstance* InModel)
 	{
-		IDMLDevice*				Device = InDevCtx->Device;
+		IDMLDevice*				Device = InModel->DevCtx->Device;
 		TComPtr<IDMLDevice1>	Device1;
 
-		Device1.FromQueryInterface(__uuidof(IDMLDevice1), InDevCtx->Device);
+		Device1.FromQueryInterface(__uuidof(IDMLDevice1), Device);
 		check(Device1);
+		
 		if (!Device1)
 		{
 			return nullptr;
 		}
 
-		if (!AddEdges(InGraph))
+		if (!AddEdges(InModel))
 		{
 			return nullptr;
 		}
@@ -402,8 +382,8 @@ public:
 		TArray<FDebugName>							DbgIntermediateNames;
 		TArray<FDebugName>							DbgOutputNames;
 
-		DbgInputNames.Reserve(InGraph.AllTensors.Num());
-		DbgIntermediateNames.Reserve(InGraph.AllTensors.Num());
+		DbgInputNames.Reserve(InModel->AllSymbolicTensorDescs.Num());
+		DbgIntermediateNames.Reserve(InModel->AllSymbolicTensorDescs.Num());
 
 		for (const FEdge& Edge : Edges)
 		{
@@ -415,7 +395,7 @@ public:
 				Input.ToNodeIndex = Edge.NodeDst;
 				Input.ToNodeInputIndex = Edge.NodeDstInput;
 				
-				const FDebugName& DbgName = DbgInputNames.Add_GetRef(InGraph.AllTensors[Edge.TensorIdx].GetName());
+				const FDebugName& DbgName = DbgInputNames.Add_GetRef(InModel->AllSymbolicTensorDescs[Edge.TensorIdx].GetName());
 				Input.Name = DbgName.Get();
 			}
 			else if (Edge.Type == EEdgeType::Output)
@@ -426,7 +406,7 @@ public:
 				Output.FromNodeIndex = Edge.NodeSrc;
 				Output.FromNodeOutputIndex = Edge.NodeSrcOutput;
 
-				const FDebugName& DbgName = DbgOutputNames.Add_GetRef(InGraph.AllTensors[Edge.TensorIdx].GetName());
+				const FDebugName& DbgName = DbgOutputNames.Add_GetRef(InModel->AllSymbolicTensorDescs[Edge.TensorIdx].GetName());
 				Output.Name = DbgName.Get();
 			}
 			else if (Edge.Type == EEdgeType::Intermediate)
@@ -438,7 +418,7 @@ public:
 				Intermediate.ToNodeIndex = Edge.NodeDst;
 				Intermediate.ToNodeInputIndex = Edge.NodeDstInput;
 
-				const FDebugName& DbgName = DbgIntermediateNames.Add_GetRef(InGraph.AllTensors[Edge.TensorIdx].GetName());
+				const FDebugName& DbgName = DbgIntermediateNames.Add_GetRef(InModel->AllSymbolicTensorDescs[Edge.TensorIdx].GetName());
 				Intermediate.Name = DbgName.Get();
 			}
 		}
@@ -501,7 +481,7 @@ public:
 		Res = Device1->CompileGraph(&Graph, DML_EXECUTION_FLAG_NONE, DML_PPV_ARGS(&Op));
 		if (FAILED(Res))
 		{
-			UE_LOG(LogNNE, Warning, TEXT("Error:Failed to compile DML graph"));
+			UE_LOG(LogNNE, Error, TEXT("Failed to compile DML graph"));
 			Op = nullptr;
 		};
 
@@ -510,7 +490,10 @@ public:
 
 private:
 
-	bool AddEdges(const FGraphDesc& InGraph)
+	//
+	//
+	//
+	bool AddEdges(const FModelInstance* InModel)
 	{
 		Edges.Reset();
 		TensorInConnCounts.Reset();
@@ -519,47 +502,47 @@ private:
 		NumInputs = 0;
 		NumOutputs = 0;
 
-		TensorInConnCounts.SetNumZeroed(InGraph.AllTensors.Num());
-		TensorOutConnCounts.SetNumZeroed(InGraph.AllTensors.Num());
-		Operators.Reset(InGraph.Operators.Num());
+		TensorInConnCounts.SetNumZeroed(InModel->AllSymbolicTensorDescs.Num());
+		TensorOutConnCounts.SetNumZeroed(InModel->AllSymbolicTensorDescs.Num());
+		Operators.Reset(InModel->GraphOperators.Num());
 
-		for (int32 Idx : InGraph.OpInputIndices)
+		for (int32 Idx : InModel->GraphOpInputIndices)
 		{
 			TensorInConnCounts[Idx] += 1;
 		}
 
-		for (int32 Idx : InGraph.OpOutputIndices)
+		for (int32 Idx : InModel->GraphOpOutputIndices)
 		{
 			TensorOutConnCounts[Idx] += 1;
 		}
 
-		if (!AddInputEdges(InGraph.InputIndices, InGraph))
+		if (!AddInputEdges(InModel->InputTensorIndices, InModel))
 		{
-			UE_LOG(LogNNE, Warning, TEXT("DMLGraphBuilder failed to add input tensors"));
+			UE_LOG(LogNNE, Error, TEXT("DMLGraphBuilder failed to add input tensors"));
 			return false;
 		}
 		
-		if (!AddInputEdges(InGraph.WeightIndices, InGraph))
+		if (!AddInputEdges(InModel->WeightTensorIndices, InModel))
 		{
-			UE_LOG(LogNNE, Warning, TEXT("DMLGraphBuilder failed to add weight tensors"));
+			UE_LOG(LogNNE, Error, TEXT("DMLGraphBuilder failed to add weight tensors"));
 			return false;
 		}
 		
-		if (!AddOutputEdges(InGraph))
+		if (!AddOutputEdges(InModel))
 		{
-			UE_LOG(LogNNE, Warning, TEXT("DMLGraphBuilder failed to add output tensors"));
+			UE_LOG(LogNNE, Error, TEXT("DMLGraphBuilder failed to add output tensors"));
 			return false;
 		}
 
-		if (!AddIntermediateEdges(InGraph))
+		if (!AddIntermediateEdges(InModel))
 		{
-			UE_LOG(LogNNE, Warning, TEXT("DMLGraphBuilder failed to add intermediate tensors"));
+			UE_LOG(LogNNE, Error, TEXT("DMLGraphBuilder failed to add intermediate tensors"));
 			return false;
 		}
 
-		if (!AddOperators(InGraph))
+		if (!AddOperators(InModel))
 		{
-			UE_LOG(LogNNE, Warning, TEXT("DMLGraphBuilder failed to add operators"));
+			UE_LOG(LogNNE, Error, TEXT("DMLGraphBuilder failed to add operators"));
 			return false;
 		}
 
@@ -570,7 +553,7 @@ private:
 			{
 				if (Edge.NodeSrcOutput == -1 || Edge.NodeDst == -1 || Edge.NodeDstInput == -1)
 				{
-					UE_LOG(LogNNE, Warning, TEXT("DMLGraphBuilder error invalid input graph edge detected for tensor:%s"), *InGraph.AllTensors[Edge.TensorIdx].GetName());
+					UE_LOG(LogNNE, Error, TEXT("DMLGraphBuilder error invalid input graph edge detected for tensor:%s"), *InModel->AllSymbolicTensorDescs[Edge.TensorIdx].GetName());
 					return false;
 				}
 			}
@@ -578,7 +561,7 @@ private:
 			{
 				if (Edge.NodeDstInput == -1 || Edge.NodeSrc == -1 || Edge.NodeSrcOutput == -1)
 				{
-					UE_LOG(LogNNE, Warning, TEXT("DMLGraphBuilder error invalid output graph edge detected for tensor:%s"), *InGraph.AllTensors[Edge.TensorIdx].GetName());
+					UE_LOG(LogNNE, Error, TEXT("DMLGraphBuilder error invalid output graph edge detected for tensor:%s"), *InModel->AllSymbolicTensorDescs[Edge.TensorIdx].GetName());
 					return false;
 				}
 			}
@@ -586,7 +569,7 @@ private:
 			{
 				if (Edge.NodeSrc == -1 || Edge.NodeSrcOutput == -1 || Edge.NodeDst == -1 || Edge.NodeDstInput == -1)
 				{
-					UE_LOG(LogNNE, Warning, TEXT("DMLGraphBuilder error invalid intermediate graph edge detected for tensor:%s"), *InGraph.AllTensors[Edge.TensorIdx].GetName());
+					UE_LOG(LogNNE, Error, TEXT("DMLGraphBuilder error invalid intermediate graph edge detected for tensor:%s"), *InModel->AllSymbolicTensorDescs[Edge.TensorIdx].GetName());
 					return false;
 				}
 			}
@@ -595,7 +578,10 @@ private:
 		return true;
 	}
 
-	bool AddInputEdges(TConstArrayView<int32> InTensorIndices, const FGraphDesc& InGraph)
+	//
+	//
+	//
+	bool AddInputEdges(TConstArrayView<int32> InTensorIndices, const FModelInstance* InModel)
 	{
 		for (int32 Idx = 0; Idx < InTensorIndices.Num(); ++Idx)
 		{
@@ -604,20 +590,23 @@ private:
 
 			if (ConnCount == 0)
 			{
-				if (InGraph.ConstantCPUIndices.Find(TensorIdx) == INDEX_NONE)
+				if (InModel->ConstantCPUTensorIndices.Find(TensorIdx) == INDEX_NONE)
 				{
-					UE_LOG(LogNNE, Warning, TEXT("DmlGraphBuilder tensor:%d has no connections"), *InGraph.AllTensors[TensorIdx].GetName());
+					UE_LOG(LogNNE, Error, TEXT("DmlGraphBuilder tensor:%d has no connections"), *InModel->AllSymbolicTensorDescs[TensorIdx].GetName());
 					return false;
 				}
 			}
 
-			for (int32 ConnIdx = 0; ConnIdx < ConnCount; ++ConnIdx)
-			{
-				AddInputEdge(TensorIdx);
-			}
+			// Check if tensor is a constant CPU
+			const int32 ConstantCPUIdx = InModel->ConstantCPUTensorIndices.Find(TensorIdx);
 
-			if (InGraph.ConstantCPUIndices.Find(TensorIdx) == INDEX_NONE)
+			if (ConstantCPUIdx == INDEX_NONE)
 			{
+				for (int32 ConnIdx = 0; ConnIdx < ConnCount; ++ConnIdx)
+				{
+					AddInputEdge(TensorIdx);
+				}
+			
 				++NumInputs;
 			}
 		}
@@ -625,16 +614,19 @@ private:
 		return true;
 	}
 
-	bool AddOutputEdges(const FGraphDesc& InGraph)
+	//
+	//
+	//
+	bool AddOutputEdges(const FModelInstance* InModel)
 	{
-		for (int32 Idx = 0; Idx < InGraph.OutputIndices.Num(); ++Idx)
+		for (int32 Idx = 0; Idx < InModel->OutputTensorIndices.Num(); ++Idx)
 		{
-			const int32 TensorIdx = InGraph.OutputIndices[Idx];
+			const int32 TensorIdx = InModel->OutputTensorIndices[Idx];
 			const int32 ConnCount = TensorOutConnCounts[TensorIdx];
 
 			if (ConnCount == 0)
 			{
-				UE_LOG(LogNNE, Warning, TEXT("DmlGraphBuilder tensor:%d has no connections"), *InGraph.AllTensors[TensorIdx].GetName());
+				UE_LOG(LogNNE, Error, TEXT("DmlGraphBuilder tensor:%d has no connections"), *InModel->AllSymbolicTensorDescs[TensorIdx].GetName());
 				return false;
 			}
 
@@ -647,17 +639,20 @@ private:
 		return true;
 	}
 
-	bool AddIntermediateEdges(const FGraphDesc& InGraph)
+	//
+	//
+	//
+	bool AddIntermediateEdges(const FModelInstance* InModel)
 	{
-		for (int32 Idx = 0; Idx < InGraph.IntermediateIndices.Num(); ++Idx)
+		for (int32 Idx = 0; Idx < InModel->IntermediateTensorIndices.Num(); ++Idx)
 		{
-			const int32 TensorIdx = InGraph.IntermediateIndices[Idx];
+			const int32 TensorIdx = InModel->IntermediateTensorIndices[Idx];
 			const int32 InputConnCount = TensorInConnCounts[TensorIdx];
 			const int32 OutputConnCount = TensorOutConnCounts[TensorIdx];
 
 			if (InputConnCount == 0 || OutputConnCount == 0)
 			{
-				UE_LOG(LogNNE, Warning, TEXT("DmlGraphBuilder tensor:%d has no connections"), *InGraph.AllTensors[TensorIdx].GetName());
+				UE_LOG(LogNNE, Error, TEXT("DmlGraphBuilder tensor:%d has no connections"), *InModel->AllSymbolicTensorDescs[TensorIdx].GetName());
 				return false;
 			}
 
@@ -670,31 +665,44 @@ private:
 		return true;
 	}
 
-	bool AddOperators(const FGraphDesc& InGraph)
+	//
+	//
+	//
+	bool AddOperators(const FModelInstance* InModel)
 	{
-		for (const FOpDesc& CurrOp : InGraph.Operators) 
+		for (const FGraphOpDesc& CurrOp : InModel->GraphOperators) 
 		{
-			DML_OPERATOR_GRAPH_NODE_DESC& OpDesc = Operators.Add_GetRef({});
-
-			OpDesc.Operator = CurrOp.Op->GetOperator();
+			DML_OPERATOR_GRAPH_NODE_DESC&	OpDesc = Operators.Add_GetRef({});
+			FOperatorDml*					Op = InModel->Operators[CurrOp.OpIndex];
+			
+			OpDesc.Operator = Op->GetOperator();
 			OpDesc.Name = CurrOp.DbgName.Get();
 
-			const int32 NodeIdx = Operators.Num() - 1;
+			const int32		NodeIdx = Operators.Num() - 1;
+			int32			InputConnIdx = 0;
 
 			for (int32 Idx = 0; Idx < CurrOp.InputCount; ++Idx)
 			{
-				const int32 TensorIdx = InGraph.OpInputIndices[Idx + CurrOp.InputStart];
+				const int32 TensorIdx = InModel->GraphOpInputIndices[Idx + CurrOp.InputStart];
 
-				if (!ConnectEdgeDst(TensorIdx, NodeIdx, Idx))
+				// Filter out the constant CPU inputs from the graph node inputs
+				if (InModel->ConstantCPUTensorIndices.Find(TensorIdx) != INDEX_NONE)
 				{
-					UE_LOG(LogNNE, Warning, TEXT("DmlGraphBuilder failed to connect intermediate edge dst for tensor:%s"), *InGraph.AllTensors[TensorIdx].GetName());
+					continue;
+				}
+
+				if (!ConnectEdgeDst(TensorIdx, NodeIdx, InputConnIdx))
+				{
+					UE_LOG(LogNNE, Error, TEXT("DmlGraphBuilder failed to connect intermediate edge dst for tensor:%s"), *InModel->AllSymbolicTensorDescs[TensorIdx].GetName());
 					return false;
 				}
+
+				++InputConnIdx;
 			}
 
 			for (int32 Idx = 0; Idx < CurrOp.OutputCount; ++Idx)
 			{
-				const int32 TensorIdx = InGraph.OpOutputIndices[Idx + CurrOp.OutputStart];
+				const int32 TensorIdx = InModel->GraphOpOutputIndices[Idx + CurrOp.OutputStart];
 				int32 ConnCount = TensorInConnCounts[TensorIdx];
 
 				if (ConnCount == 0)
@@ -707,7 +715,7 @@ private:
 				{
 					if (!ConnectEdgeSrc(TensorIdx, NodeIdx, Idx))
 					{
-						UE_LOG(LogNNE, Warning, TEXT("DmlGraphBuilder failed to connect intermediate edge src for tensor:%s"), *InGraph.AllTensors[TensorIdx].GetName());
+						UE_LOG(LogNNE, Error, TEXT("DmlGraphBuilder failed to connect intermediate edge src for tensor:%s"), *InModel->AllSymbolicTensorDescs[TensorIdx].GetName());
 						return false;
 					}
 				}
@@ -717,6 +725,9 @@ private:
 		return true;
 	}
 
+	//
+	//
+	//
 	void AddInputEdge(int32 TensorIdx)
 	{
 		Edges.Emplace(
@@ -726,6 +737,9 @@ private:
 		);
 	}
 
+	//
+	//
+	//
 	void AddOutputEdge(int32 TensorIdx)
 	{
 		Edges.Emplace(
@@ -737,6 +751,9 @@ private:
 		++NumOutputs;
 	}
 
+	//
+	//
+	//
 	void AddIntermediateEdge(int32 TensorIdx, int32 NodeSrc, int32 NodeSrcOutput)
 	{
 		Edges.Emplace(
@@ -747,6 +764,9 @@ private:
 			);
 	}
 
+	//
+	//
+	//
 	bool ConnectEdgeDst(int32 TensorIdx, int32 NodeDst, int32 NodeDstInput)
 	{
 		FEdge* StartEdge = Edges.FindByPredicate(
@@ -773,6 +793,9 @@ private:
 		return bFoundEdge;
 	}
 
+	//
+	//
+	//
 	bool ConnectEdgeSrc(int32 TensorIdx, int32 NodeSrc, int32 NodeSrcOutput)
 	{
 		FEdge* StartEdge =
@@ -820,16 +843,10 @@ FModelInstance::FModelInstance()
 //
 FModelInstance::~FModelInstance()
 {
-
 #if STATS
 	DEC_MEMORY_STAT_BY(STAT_MemSizeWeights, MemSizeWeights);
 	DEC_MEMORY_STAT_BY(STAT_MemSizeTemp, MemSizeTemp);
 	DEC_MEMORY_STAT_BY(STAT_MemSizePersist, MemSizePersist);
-
-//	if (CompiledOp)
-//	{
-//		DEC_DWORD_STAT(STAT_InferenceCount);
-//	}
 #endif
 }
 
@@ -856,37 +873,27 @@ bool FModelInstance::Init(TConstArrayView<uint8> ModelData, FDmlDeviceContext* I
 	HRESULT Res = DevCtx->Device->CreateOperatorInitializer(0, nullptr, DML_PPV_ARGS(&OpInit));
 	if (!OpInit)
 	{
-		UE_LOG(LogNNE, Warning, TEXT("Error:Failed to create DML operator initializer"));
+		UE_LOG(LogNNE, Error, TEXT("Failed to create DML operator initializer"));
 		return false;
 	}
 
-	// DirectML requires all tensors to be concrete
-	// Notes: to handle dynamic tensor desc, op should init from symbolic shapes
-	TArray<NNE::Internal::FTensor>	Tensors;
-
-	Tensors.Reset(AllSymbolicTensorDescs.Num());
-	for (const NNE::FTensorDesc& TensorDesc : AllSymbolicTensorDescs)
-	{
-		Tensors.Emplace(NNE::Internal::FTensor::MakeFromSymbolicDesc(TensorDesc));
-	}
-
-	TArray<FGraphBuilder::FOpDesc>	DmlGraphOperators;
-	TArray<int32>					OpInputIndices;
-	TArray<int32>					OpOutputIndices;
-	uint64							TensorDataSize = 0;
+	GraphOperators.Reset();
+	GraphOpInputIndices.Reset();
+	GraphOpOutputIndices.Reset();
+	MemSizeWeights = 0;
 
 	// Loop over all operators in the model and create them
 	for (int32 Idx = 0; Idx < Format.Operators.Num(); ++Idx)
 	{
 		const FString TypeName = Format.Operators[Idx].TypeName;
 
-		FGraphBuilder::FOpDesc					OpDesc;
-		TArray<NNE::Internal::FTensor>		OpInputTensors;
-		TArray<NNE::Internal::FTensor>		OpOutputTensors;
-		NNE::FAttributeMap					AttributeMap;
+		FGraphOpDesc				OpDesc;
+		TArray<NNE::FTensorDesc>	OpInputTensors;
+		TArray<NNE::FTensorDesc>	OpOutputTensors;
+		NNE::FAttributeMap			AttributeMap;
 
-		OpDesc.InputStart = OpInputIndices.Num();
-		OpDesc.OutputStart = OpOutputIndices.Num();
+		OpDesc.InputStart = GraphOpInputIndices.Num();
+		OpDesc.OutputStart = GraphOpOutputIndices.Num();
 
 		for (int32 InputTensorIndex : Format.Operators[Idx].InTensors)
 		{
@@ -894,26 +901,19 @@ bool FModelInstance::Init(TConstArrayView<uint8> ModelData, FDmlDeviceContext* I
 
 			if (WeightTensorIdx >= 0)
 			{
-				TConstArrayView<uint8> TensorData = OpInputTensors.Emplace_GetRef(WeightTensorRDGs[WeightTensorIdx]).GetPreparedData<uint8>();
+				TConstArrayView<uint8> TensorData = WeightTensorRDGs[WeightTensorIdx].GetPreparedData<uint8>();
 				
-				TensorDataSize += Align(TensorData.Num(), DML_MINIMUM_BUFFER_TENSOR_ALIGNMENT);
-			}
-			else
-			{
-				NNE::FTensorDesc SymbolicTensorDesc = AllSymbolicTensorDescs[InputTensorIndex];
-				// Notes: to handle dynamic tensor desc, op should init from symbolic shapes
-				OpInputTensors.Emplace(NNE::Internal::FTensor::MakeFromSymbolicDesc(SymbolicTensorDesc));
+				MemSizeWeights += Align(TensorData.Num(), DML_MINIMUM_BUFFER_TENSOR_ALIGNMENT);
 			}
 
-			OpInputIndices.Emplace(InputTensorIndex);
+			OpInputTensors.Emplace(AllSymbolicTensorDescs[InputTensorIndex]);
+			GraphOpInputIndices.Emplace(InputTensorIndex);
 		}
 
 		for (int32 OutputTensorIndex : Format.Operators[Idx].OutTensors)
 		{
-			NNE::FTensorDesc SymbolicTensorDesc = AllSymbolicTensorDescs[OutputTensorIndex];
-			// Notes: to handle dynamic tensor desc, op should init from symbolic shapes
-			OpOutputTensors.Emplace(NNE::Internal::FTensor::MakeFromSymbolicDesc(SymbolicTensorDesc));
-			OpOutputIndices.Emplace(OutputTensorIndex);
+			OpOutputTensors.Emplace(AllSymbolicTensorDescs[OutputTensorIndex]);
+			GraphOpOutputIndices.Emplace(OutputTensorIndex);
 		}
 
 		for (const FNNEFormatAttributeDesc& Desc : Format.Operators[Idx].Attributes)
@@ -921,70 +921,60 @@ bool FModelInstance::Init(TConstArrayView<uint8> ModelData, FDmlDeviceContext* I
 			AttributeMap.SetAttribute(Desc.Name, Desc.Value);
 		}
 
-		OpDesc.Op = OpCreate(TypeName, OpInputTensors, OpOutputTensors, AttributeMap);
+		FOperatorDml* Op = OpCreate(TypeName, OpInputTensors, OpOutputTensors, AttributeMap);
 
-		if (!OpDesc.Op)
+		if (!Op)
 		{
-			UE_LOG(LogNNE, Warning, TEXT("Error:Failed to create DML operator:%s"), *TypeName);
+			UE_LOG(LogNNE, Error, TEXT("Failed to create DML operator:%s"), *TypeName);
 			return false;
 		}
 
+		OpDesc.OpIndex = Operators.Num();
+		Operators.Add(Op);
+
 		// Remap inputs
-		TConstArrayView<int32> OpRemappedInputs = OpDesc.Op->GetRemappedInputs();
+		TConstArrayView<int32> OpRemappedInputs = Op->GetRemappedInputs();
 
 		for (int32 InputIdx = 0; InputIdx < OpRemappedInputs.Num(); ++InputIdx)
 		{
-			OpInputIndices[OpDesc.InputStart + InputIdx] = OpRemappedInputs[InputIdx];
+			GraphOpInputIndices[OpDesc.InputStart + InputIdx] = OpRemappedInputs[InputIdx];
 		}
 
-		// Filter out the constant CPU inputs from the graph node inputs
-		TConstArrayView<int32> OpConstantCPUInputs = OpDesc.Op->GetConstantCPUInputs();
+		OpDesc.InputCount = OpInputTensors.Num();
+		OpDesc.OutputCount = OpOutputTensors.Num();
+		OpDesc.DbgName = TypeName;
+
+		GraphOperators.Emplace(OpDesc);
+
+		// Add constant CPU tensor indices
+		TConstArrayView<int32> OpConstantCPUInputs = Op->GetConstantCPUInputs();
 
 		for (int32 InputIdx = OpConstantCPUInputs.Num() - 1; InputIdx >= 0; --InputIdx)
 		{
 			const int32 ConstIdx = OpConstantCPUInputs[InputIdx];
-			const int32 TensorIdx = OpInputIndices[OpDesc.InputStart + ConstIdx];
-
-			OpInputTensors.RemoveAt(ConstIdx);
-			OpInputIndices.RemoveAt(OpDesc.InputStart + ConstIdx);
+			const int32 TensorIdx = GraphOpInputIndices[OpDesc.InputStart + ConstIdx];
 
 			ConstantCPUTensorIndices.AddUnique(TensorIdx);
 		}
 
-		OpDesc.InputCount = OpInputTensors.Num();
-		OpDesc.OutputCount = OpOutputTensors.Num();
-		OpDesc.DbgName = TypeName;
-
-		DmlGraphOperators.Emplace(OpDesc);
-	}
-
-	FGraphBuilder				DmlGraphBuilder;
-	FGraphBuilder::FGraphDesc	DmlGraphDesc;
-
-	DmlGraphDesc.AllTensors = Tensors;
-	DmlGraphDesc.InputIndices = InputTensorIndices;
-	DmlGraphDesc.OutputIndices = OutputTensorIndices;
-	DmlGraphDesc.WeightIndices = WeightTensorIndices;
-	DmlGraphDesc.ConstantCPUIndices = ConstantCPUTensorIndices;
-	DmlGraphDesc.IntermediateIndices = IntermediateTensorIndices;
-	DmlGraphDesc.WeightTensors = WeightTensorRDGs;
-	DmlGraphDesc.Operators = DmlGraphOperators;
-	DmlGraphDesc.OpInputIndices = OpInputIndices;
-	DmlGraphDesc.OpOutputIndices = OpOutputIndices;
-
-	CompiledOp = DmlGraphBuilder.Compile(DevCtx, DmlGraphDesc);
-	if (!CompiledOp)
-	{
-		return false;
+		// FIXME: TODO: Find a way to test/validate if tensor is constant, for now we can only do that in FOperatorDml::PrepareOutputs()
+		//for (int32 TensorIdx : ConstantCPUTensorIndices)
+		//{
+		//	if (!AllTensorRDGRefs[TensorIdx]->HasPreparedData())
+		//	{
+		//		UE_LOG(LogNNE, Error, TEXT("DirectML requires tensor %s needs to be constant"), *AllTensorRDGRefs[TensorIdx]->GetName());
+		//		return false;
+		//	}
+		//}
 	}
 
-	return InitCompiledOp(OpInputIndices, TensorDataSize);
+	return true;	
 }
 
 //
 //
 //
-bool FModelInstance::InitCompiledOp(TConstArrayView<int32> OpInputIndices, uint64 TensorDataSize)
+bool FModelInstance::InitCompiledOp()
 {
 	static constexpr EBufferUsageFlags	WeightBuffUsage = BUF_UnorderedAccess;
 	static constexpr ERHIAccess			WeightBuffAccess = ERHIAccess::UAVMask;
@@ -1002,7 +992,7 @@ bool FModelInstance::InitCompiledOp(TConstArrayView<int32> OpInputIndices, uint6
 	Res = OpInit->Reset(UE_ARRAY_COUNT(CompiledOps), CompiledOps);
 	if (FAILED(Res))
 	{
-		UE_LOG(LogNNE, Warning, TEXT("Error:Failed to reset DML operator initializer"));
+		UE_LOG(LogNNE, Error, TEXT("Failed to reset DML operator initializer"));
 		return false;
 	}
 
@@ -1020,7 +1010,7 @@ bool FModelInstance::InitCompiledOp(TConstArrayView<int32> OpInputIndices, uint6
 	Res = DevCtx->D3D12Device->CreateDescriptorHeap(&HeapDesc, IID_PPV_ARGS(&DescHeap));
 	if (!DescHeap)
 	{
-		UE_LOG(LogNNE, Warning, TEXT("Failed to create descriptor heap, res:%x"), Res);
+		UE_LOG(LogNNE, Error, TEXT("Failed to create descriptor heap, res:%x"), Res);
 		return false;
 	}
 
@@ -1032,7 +1022,6 @@ bool FModelInstance::InitCompiledOp(TConstArrayView<int32> OpInputIndices, uint6
 		return false;
 	}
 
-	MemSizeWeights = TensorDataSize;
 	MemSizeTemp = ExecBindProps.TemporaryResourceSize;
 	MemSizePersist = ExecBindProps.PersistentResourceSize;
 
@@ -1043,8 +1032,7 @@ bool FModelInstance::InitCompiledOp(TConstArrayView<int32> OpInputIndices, uint6
 		[
 			this, 
 			Signal, 
-			InitTempMemSize = InitBindProps.TemporaryResourceSize,
-			TensorDataSize
+			InitTempMemSize = InitBindProps.TemporaryResourceSize
 		]
 		(FRHICommandListImmediate& RHICmdList)
 		{
@@ -1061,12 +1049,12 @@ bool FModelInstance::InitCompiledOp(TConstArrayView<int32> OpInputIndices, uint6
 			TArray<CD3DX12_RESOURCE_BARRIER, TInlineAllocator<MaxNumInputs>>	Barriers;
 			FGPUFenceRHIRef	UploadFence = nullptr;
 
-			if (TensorDataSize)
+			if (MemSizeWeights)
 			{
 				UploadFence = RHICreateGPUFence(TEXT("FDmlModelInstance_UploadFence"));
 
-				FBufferRHIRef	UploadBuff = CreateRHIBuffer(RHICmdList, TensorDataSize, BUF_ShaderResource | BUF_Dynamic | BUF_FastVRAM, ERHIAccess::CopySrc, TEXT("FDmlModelInstance_UploadBuffer"));
-				uint8*			UploadBuffPtr = static_cast<uint8*>(RHICmdList.LockBuffer(UploadBuff, 0, TensorDataSize, RLM_WriteOnly_NoOverwrite));
+				FBufferRHIRef	UploadBuff = CreateRHIBuffer(RHICmdList, MemSizeWeights, BUF_ShaderResource | BUF_Dynamic | BUF_FastVRAM, ERHIAccess::CopySrc, TEXT("FDmlModelInstance_UploadBuffer"));
+				uint8*			UploadBuffPtr = static_cast<uint8*>(RHICmdList.LockBuffer(UploadBuff, 0, MemSizeWeights, RLM_WriteOnly_NoOverwrite));
 				uint64			UploadOffset = 0;
 
 				for (int32 WeightIdx = 0; WeightIdx < WeightTensorIndices.Num(); ++WeightIdx)
@@ -1100,7 +1088,7 @@ bool FModelInstance::InitCompiledOp(TConstArrayView<int32> OpInputIndices, uint6
 				RHICmdList.UnlockBuffer(UploadBuff);
 				RHICmdList.WriteGPUFence(UploadFence);
 
-				INC_MEMORY_STAT_BY(STAT_MemSizeWeights, TensorDataSize);
+				INC_MEMORY_STAT_BY(STAT_MemSizeWeights, MemSizeWeights);
 			}
 
 			if (MemSizePersist)
@@ -1299,31 +1287,32 @@ void FModelInstance::AddDispatchOps_RenderThread(FRDGBuilder& GraphBuilder)
 //
 // Create operator
 //
-FOperatorDml* FModelInstance::OpCreate(const FString& OpName, TArrayView<const NNE::Internal::FTensor> InputTensorDescs, TArrayView<const NNE::Internal::FTensor> OutputTensorDescs, const NNE::FAttributeMap& Attributes)
+FOperatorDml* FModelInstance::OpCreate(const FString& OpName, TConstArrayView<NNE::FTensorDesc> Inputs, TConstArrayView<NNE::FTensorDesc> Outputs, const NNE::FAttributeMap& Attributes)
 {
 	FOperatorRegistryDml::OperatorCreateFunc CreateFn = FOperatorRegistryDml::Get()->OpFind(OpName);
 
 	if (!CreateFn)
 	{
-		UE_LOG(LogNNE, Warning, TEXT("Error:failed to find DML operator:%s"), *OpName);
+		UE_LOG(LogNNE, Error, TEXT("Failed to find DML operator:%s"), *OpName);
 		return nullptr;
 	}
 
 	FOperatorDml* Op = CreateFn();
 
-	if (!Op->Initialize(DevCtx->Device, InputTensorDescs, OutputTensorDescs, Attributes))
+	if (!Op->Initialize(Inputs, Outputs, Attributes))
 	{
 		delete Op;
 
-		UE_LOG(LogNNE, Warning, TEXT("Error:Failed to initialize DML operator:%s"), *OpName);
+		UE_LOG(LogNNE, Error, TEXT("Failed to initialize DML operator:%s"), *OpName);
 		return nullptr;
 	}
 
-	Op->GetOperator()->SetName(*OpName);
-
 	return Op;
 }
 
+//
+//
+//
 FBufferRHIRef FModelInstance::CreateRHIBuffer(FRHICommandListImmediate& RHICmdList, uint32 Size, EBufferUsageFlags Usage, ERHIAccess Access, const TCHAR* DbgName)
 {
 	FBufferRHIRef Buff = nullptr;
@@ -1339,53 +1328,75 @@ FBufferRHIRef FModelInstance::CreateRHIBuffer(FRHICommandListImmediate& RHICmdLi
 	return Buff;
 }
 
-ID3D12Resource* FModelInstance::CreateD3D12Buffer(uint32 Size, D3D12_RESOURCE_STATES ResourceState, D3D12_HEAP_TYPE HeapType, const TCHAR* DebugName)
+//
+//
+//
+int FModelInstance::PrepareTensorShapesAndData()
 {
-	ID3D12Resource* Resource = nullptr;
-
-	D3D12_RESOURCE_DESC ResourceDesc = CD3DX12_RESOURCE_DESC::Buffer(Size,D3D12_RESOURCE_FLAG_ALLOW_UNORDERED_ACCESS);
-	CD3DX12_HEAP_PROPERTIES HeapProps = CD3DX12_HEAP_PROPERTIES(HeapType);
-	HRESULT Res;
-
-	Res = DevCtx->D3D12Device->CreateCommittedResource(
-		&HeapProps,
-		D3D12_HEAP_FLAG_NONE,
-		&ResourceDesc,
-		ResourceState,
-		nullptr,
-		IID_PPV_ARGS(&Resource));
+	check(AllTensorRDGRefs.Num() == AllSymbolicTensorDescs.Num());
 
-	if (FAILED(Res))
+	if (Operators.Num() == 0)
 	{
-		UE_LOG(LogNNE, Warning, TEXT("Error:FDmlModelInstance failed to create D3D12 resource"));
-		return nullptr;
+		UE_LOG(LogNNE, Error, TEXT("No operators in model"));
+		return -1;
 	}
 
-	if (Resource && DebugName)
+	// Loop through all operators:
+	// 1. Run shape inference with PrepareOutputs()
+	// 2. Create DML operator
+	TArray<NNE::Internal::FTensorRef>	OpInputs;
+	TArray<NNE::Internal::FTensorRef>	OpOutputs;
+
+	for (const FGraphOpDesc& OpDesc : GraphOperators)
 	{
-		Resource->SetName(DebugName);
-	}
+		OpInputs.Reset();
+		OpOutputs.Reset();
 
-	return Resource;
-}
+		FOperatorDml* Op = Operators[OpDesc.OpIndex];
 
-//
-//
-//
-int FModelInstance::PrepareTensorShapesAndData()
-{
-	for (NNE::FTensorDesc SymbolicTensorDesc : AllSymbolicTensorDescs)
-	{
-		if (!SymbolicTensorDesc.GetShape().IsConcrete())
+		for (int32 Idx = 0; Idx < OpDesc.InputCount; ++Idx)
+		{
+			const int32 InputIdx = GraphOpInputIndices[OpDesc.InputStart + Idx];
+			OpInputs.Emplace(AllTensorRDGRefs[InputIdx]);
+		}
+
+		for (int32 Idx = 0; Idx < OpDesc.OutputCount; ++Idx)
+		{
+			const int32 OutputIdx = GraphOpOutputIndices[OpDesc.OutputStart + Idx];
+			OpOutputs.Emplace(AllTensorRDGRefs[OutputIdx]);
+		}
+
+		if (Op->PrepareOutputs(OpInputs, OpOutputs) != 0)
 		{
-			UE_LOG(LogNNE, Warning, TEXT("FDmlModelInstance does not support model with variable shapes yet."));
 			return -1;
 		}
+
+		if (!Op->Create(DevCtx->Device, OpInputs, OpOutputs))
+		{
+			return -1;
+		}
+	}
+
+	FGraphBuilder DmlGraphBuilder;
+
+	CompiledOp = DmlGraphBuilder.Compile(this);
+	if (!CompiledOp)
+	{
+		UE_LOG(LogNNE, Error, TEXT("Failed to compile DML graph"));
+		return -1;
+	}
+
+	if (!InitCompiledOp())
+	{
+		return -1;
 	}
 
 	return 0;
 }
 
+//
+//
+//
 TSharedPtr<NNE::IModelInstanceRDG> FModel::CreateModelInstanceRDG()
 {
 	FModelInstance* ModelInstance = new FModelInstance();
@@ -1401,6 +1412,9 @@ TSharedPtr<NNE::IModelInstanceRDG> FModel::CreateModelInstanceRDG()
 	return TSharedPtr<NNE::IModelInstanceRDG>(IModelInstance);
 }
 
+//
+//
+//
 FModel::FModel(const TSharedPtr<NNE::FSharedModelData>& InModelData, FDmlDeviceContext* InDevCtx)
 	: ModelData(InModelData), DevCtx(InDevCtx)
 {
diff --git a/Engine/Plugins/Experimental/NNERuntimeRDG/Source/NNERuntimeRDG/Private/Dml/NNEDmlModel.h b/Engine/Plugins/Experimental/NNERuntimeRDG/Source/NNERuntimeRDG/Private/Dml/NNEDmlModel.h
--- a/Engine/Plugins/Experimental/NNERuntimeRDG/Source/NNERuntimeRDG/Private/Dml/NNEDmlModel.h
+++ b/Engine/Plugins/Experimental/NNERuntimeRDG/Source/NNERuntimeRDG/Private/Dml/NNEDmlModel.h
@@ -45,7 +45,42 @@ class FModelInstance : public FModelInstanceRDG
 {
 	class FGraphBuilder;
 	class FBindingTable;
-	class FDebugName;
+	
+	friend class FGraphBuilder;
+
+	//
+	//
+	//
+	class FDebugName
+	{
+		static constexpr int32 Size = 128;
+
+	public:
+
+		FDebugName();
+		FDebugName(const FString& InStr);
+		FDebugName(FStringView InStr);
+
+		const char* Get() const;
+
+	private:
+
+		char	Str[Size];
+		int32	Length;
+	};
+
+	//
+	//
+	//
+	struct FGraphOpDesc
+	{
+		int32			OpIndex;
+		int32			InputStart;
+		int32			InputCount;
+		int32			OutputStart;
+		int32			OutputCount;
+		FDebugName		DbgName;
+	};
 
 public:
 
@@ -61,38 +96,47 @@ protected:
 
 private:
 
-	bool InitCompiledOp(TConstArrayView<int32> OpInputIndices, uint64 TensorDataSize);
+	FOperatorDml* OpCreate(const FString& Name, TConstArrayView<NNE::FTensorDesc> Inputs, TConstArrayView<NNE::FTensorDesc> Outputs, const NNE::FAttributeMap& Attributes);
 
-	FOperatorDml* OpCreate(const FString& Name, TArrayView<const NNE::Internal::FTensor> InputTensorDesc, TArrayView<const NNE::Internal::FTensor> OutputTensorDescs, const NNE::FAttributeMap& Attributes);
+	bool InitCompiledOp();
 
 	FBufferRHIRef CreateRHIBuffer(FRHICommandListImmediate& RHICmdList, uint32 Size, EBufferUsageFlags Usage, ERHIAccess Access, const TCHAR* DbgName);
 	ID3D12Resource* CreateD3D12Buffer(uint32 Size, D3D12_RESOURCE_STATES ResourceState = D3D12_RESOURCE_STATE_COMMON, D3D12_HEAP_TYPE HeapType = D3D12_HEAP_TYPE_DEFAULT, const TCHAR* DebugName = nullptr);
 
-	//Note: This should go into RDG
+	//NOTE: This should go into RDG
 	static constexpr int32 MaxNumInputs = 512;
 	static constexpr int32 MaxNumOutputs = 4;
 
 	using FRHIBufferInputArray = TArray<FRHIBuffer*, TInlineAllocator<MaxNumInputs>>;
 	using FRHIBufferOutputArray = TArray<FRHIBuffer*, TInlineAllocator<MaxNumOutputs>>;
 
+	FDmlDeviceContext*					DevCtx;
+	TArray<FOperatorDml*>				Operators;
+	TArray<FGraphOpDesc>				GraphOperators;
+	TArray<int32>						GraphOpInputIndices;
+	TArray<int32>						GraphOpOutputIndices;
+
 	TComPtr<IDMLOperatorInitializer>	OpInit;
 	TComPtr<IDMLCompiledOperator>		CompiledOp;
-	FDmlDeviceContext*					DevCtx;
+	
+	TArray<int32>						ConstantCPUTensorIndices;
+
+	ID3D12DynamicRHI*					DynamicRHI;
 	TUniquePtr<FBindingTable>			BindingTable;
 	TComPtr<ID3D12DescriptorHeap>		DescHeap;
 	uint32								DescCount;
 	uint32								DescSize;
 
-	TArray<int32>						ConstantCPUTensorIndices;
-
 	FBufferRHIRef						PersistBuff;
 	FBufferRHIRef						TempBuff;
 	uint64								MemSizeWeights;
 	uint64								MemSizeTemp;
 	uint64								MemSizePersist;
-	ID3D12DynamicRHI*					DynamicRHI;
 };
 
+//
+//
+//
 class FModel : public NNE::IModelRDG
 {
 public:
diff --git a/Engine/Plugins/Experimental/NNERuntimeRDG/Source/NNERuntimeRDG/Private/Dml/NNEDmlOperator.cpp b/Engine/Plugins/Experimental/NNERuntimeRDG/Source/NNERuntimeRDG/Private/Dml/NNEDmlOperator.cpp
--- a/Engine/Plugins/Experimental/NNERuntimeRDG/Source/NNERuntimeRDG/Private/Dml/NNEDmlOperator.cpp
+++ b/Engine/Plugins/Experimental/NNERuntimeRDG/Source/NNERuntimeRDG/Private/Dml/NNEDmlOperator.cpp
@@ -13,22 +13,33 @@ namespace UE::NNERuntimeRDG::Private::Dml
 namespace Util
 {
 
-static bool IsSameShape(TConstArrayView<uint32> Left, TConstArrayView<uint32> Right)
+//
+//
+//
+inline bool IsSameShape(TConstArrayView<uint32> Left, TConstArrayView<uint32> Right)
 {
 	return Algo::Compare(Left, Right);
 }
 
-bool IsSameShape(const NNE::Internal::FTensor& Left, const NNE::Internal::FTensor& Right)
+//
+//
+//
+inline bool IsSameShape(const NNE::Internal::FTensor& Left, const NNE::Internal::FTensor& Right)
 {
 	return IsSameShape(Left.GetShape().GetData(), Right.GetShape().GetData());
 }
 
-bool IsSameShape(const FTensorDescDml& Left, const FTensorDescDml& Right)
+//
+//
+//
+inline bool IsSameShape(const FTensorDescDml& Left, const FTensorDescDml& Right)
 {
 	return IsSameShape(Left.GetSizes(), Right.GetSizes());
 }
 
-
+//
+//
+//
 DML_TENSOR_DATA_TYPE GetTensorDataType(ENNETensorDataType DataType)
 {
 	switch (DataType)
@@ -73,6 +84,9 @@ DML_TENSOR_DATA_TYPE GetTensorDataType(ENNETensorDataType DataType)
 
 } // namespace Util 
 
+//
+//
+//
 FTensorDescDml::FTensorDescDml()
 	: MinTensorRank(1)
 	, MaxTensorRank(DML_TENSOR_DIMENSION_COUNT_MAX1)
@@ -84,6 +98,9 @@ FTensorDescDml::FTensorDescDml()
 	Desc = DML_TENSOR_DESC{};
 }
 
+//
+//
+//
 FTensorDescDml& FTensorDescDml::SetTensorRank(int32 MinRank, int32 MaxRank)
 {
 	MinTensorRank = MinRank;
@@ -91,6 +108,9 @@ FTensorDescDml& FTensorDescDml::SetTensorRank(int32 MinRank, int32 MaxRank)
 	return *this;
 }
 
+//
+//
+//
 FTensorDescDml& FTensorDescDml::SetShape(TConstArrayView<uint32> Shape)
 {
 	check(!bIsValidated);
@@ -113,6 +133,9 @@ FTensorDescDml& FTensorDescDml::SetShape(TConstArrayView<uint32> Shape)
 	return *this;
 }
 
+//
+//
+//
 FTensorDescDml& FTensorDescDml::SetShape(TConstArrayView<uint32> Shape, int32 Rank)
 {
 	check(!bIsValidated);
@@ -145,14 +168,16 @@ FTensorDescDml& FTensorDescDml::SetShape(TConstArrayView<uint32> Shape, int32 Ra
 	return *this;
 }
 
-
+//
+//
+//
 FTensorDescDml& FTensorDescDml::SetShape(TConstArrayView<uint32> Shape, TConstArrayView<uint32> BroadcastShape)
 {
 	check(!bIsValidated);
 	
 	if (!bIsValidated)
 	{
-		const uint32 TargetDimension = BroadcastShape.Num() != -1 ? BroadcastShape.Num() : Shape.Num();
+		const uint32 TargetDimension = BroadcastShape.Num() ? BroadcastShape.Num() : Shape.Num();
 		checkf(BroadcastShape.Num() >= Shape.Num(), TEXT("Can't broadcast tensor from rank %d to rank %d, should be inferior or equal."), Shape.Num(), TargetDimension);
 
 		Sizes.SetNum(TargetDimension);
@@ -181,6 +206,9 @@ FTensorDescDml& FTensorDescDml::SetShape(TConstArrayView<uint32> Shape, TConstAr
 	return *this;
 }
 
+//
+//
+//
 FTensorDescDml& FTensorDescDml::SetShape1D(uint32 Dimension, int32 Rank)
 {
 	check(!bIsValidated);
@@ -208,6 +236,9 @@ FTensorDescDml& FTensorDescDml::SetShape1D(uint32 Dimension, int32 Rank)
 	return *this;
 }
 
+//
+//
+//
 FTensorDescDml& FTensorDescDml::SetStridesFromShape(TConstArrayView<uint32> Shape)
 {
 	check(!bIsValidated);
@@ -227,6 +258,9 @@ FTensorDescDml& FTensorDescDml::SetStridesFromShape(TConstArrayView<uint32> Shap
 	return *this;
 }
 
+//
+//
+//
 FTensorDescDml& FTensorDescDml::SetStrides(TConstArrayView<uint32> InStrides)
 {
 	check(!bIsValidated);
@@ -239,6 +273,9 @@ FTensorDescDml& FTensorDescDml::SetStrides(TConstArrayView<uint32> InStrides)
 	return *this;
 }
 
+//
+//
+//
 FTensorDescDml& FTensorDescDml::SetDataType(ENNETensorDataType DataType)
 {
 	check(!bIsValidated);
@@ -258,6 +295,9 @@ FTensorDescDml& FTensorDescDml::SetDataType(ENNETensorDataType DataType)
 	return *this;
 }
 
+//
+//
+//
 FTensorDescDml& FTensorDescDml::SetDataOwnedByDml(bool bSetOwnedByDml)
 {
 	check(!bIsValidated);
@@ -270,6 +310,9 @@ FTensorDescDml& FTensorDescDml::SetDataOwnedByDml(bool bSetOwnedByDml)
 	return *this;
 }
 
+//
+//
+//
 bool FTensorDescDml::Validate()
 {
 	check(!bIsValidated);
@@ -323,6 +366,9 @@ bool FTensorDescDml::Validate()
 	return bIsValidated;
 }
 
+//
+//
+//
 uint64 FTensorDescDml::CalculateBufferSize(uint64 ElemSizeInBytes)
 {
 	if (ElemSizeInBytes == 0)
diff --git a/Engine/Plugins/Experimental/NNERuntimeRDG/Source/NNERuntimeRDG/Private/Dml/NNEDmlOperator.h b/Engine/Plugins/Experimental/NNERuntimeRDG/Source/NNERuntimeRDG/Private/Dml/NNEDmlOperator.h
--- a/Engine/Plugins/Experimental/NNERuntimeRDG/Source/NNERuntimeRDG/Private/Dml/NNEDmlOperator.h
+++ b/Engine/Plugins/Experimental/NNERuntimeRDG/Source/NNERuntimeRDG/Private/Dml/NNEDmlOperator.h
@@ -10,6 +10,9 @@
 #include "NNEAttributeMap.h"
 #include "NNERuntimeRDGBase.h"
 
+//
+//
+//
 #define NNE_DML_REGISTER_OP(OpName) \
 struct FDmlOperator##OpName##Registrator \
 { \
@@ -29,12 +32,18 @@ static constexpr uint32_t NcdhwDimensionCount = 5;
 static constexpr uint32_t NcdhwSpatialDimensionCount = 3;
 static constexpr uint32_t NonspatialDimensionCount = 2; // The batch and channel dimensions of NCW, NCHW, NCDHW....
 
+//
+//
+//
 template<typename T>
 inline TArrayView<T> MakeEmptyArrayView()
 {
 	return MakeArrayView(static_cast<T*>(nullptr), 0);
 }
 
+//
+//
+//
 template<typename T>
 inline TConstArrayView<T> MakeEmptyConstArrayView()
 {
@@ -64,24 +73,30 @@ using FSmallArray = TArray<T, TInlineAllocator<NNE::FTensorShape::MaxRank>>;
 using FSmallIntArray = TArray<int32, TInlineAllocator<NNE::FTensorShape::MaxRank>>;
 using FSmallUIntArray = TArray<uint32, TInlineAllocator<NNE::FTensorShape::MaxRank>>;
 
+//
+//
+//
 template<typename InputType, typename OutputType>
 inline bool IsOverflowing(InputType Input)
 {
 	OutputType Output = static_cast<OutputType>(Input);
-	if(Input != static_cast<InputType>(Output))
+	if (Input != static_cast<InputType>(Output))
 	{
 		return true;
 	}
 	return false;
 }
 
+//
+//
+//
 template<typename InputType, typename OutputType>
 inline bool ConvertArrayViewNoOverflow(TConstArrayView<InputType> InputView, TArrayView<OutputType>& OutputView)
 {
 	OutputView = MakeArrayView((OutputType*) InputView.GetData(), InputView.Num());
-	for(int32 Idx = 0; Idx < InputView.Num(); ++Idx)
+	for (int32 Idx = 0; Idx < InputView.Num(); ++Idx)
 	{
-		if(InputView[Idx] != static_cast<InputType>(OutputView[Idx]))
+		if (InputView[Idx] != static_cast<InputType>(OutputView[Idx]))
 		{
 			return false;
 		}
@@ -89,6 +104,9 @@ inline bool ConvertArrayViewNoOverflow(TConstArrayView<InputType> InputView, TAr
 	return true;
 }
 
+//
+//
+//
 template<typename OutputType, typename AllocatorType>
 inline bool GetArrayAttributeNoOverflow(
 	const FNNEAttributeValue* Attr, 
@@ -98,31 +116,32 @@ inline bool GetArrayAttributeNoOverflow(
 {
 	if (Attr)
 	{
-
 		TArrayView<OutputType> ConvertedView;
 		TArray<int32> IntArray;
 		TArray<float> FloatArray;
 
-		switch(Attr->GetType())
+		switch (Attr->GetType())
 		{
 		case ENNEAttributeDataType::Int32Array:
 			{
 				IntArray = Attr->GetValue<TArray<int32>>();
-				if(!ConvertArrayViewNoOverflow(TConstArrayView<int32>(IntArray), ConvertedView))
+				if (!ConvertArrayViewNoOverflow(TConstArrayView<int32>(IntArray), ConvertedView))
 				{
 					return false;
 				}
 			}
 			break;
+
 		case ENNEAttributeDataType::FloatArray:
 			{
 				FloatArray = Attr->GetValue<TArray<float>>();
-				if(!ConvertArrayViewNoOverflow(TConstArrayView<float>(FloatArray), ConvertedView))
+				if (!ConvertArrayViewNoOverflow(TConstArrayView<float>(FloatArray), ConvertedView))
 				{
 					return false;
 				}
 			}
 			break;
+
 		default:
 			return false;
 		}
@@ -253,6 +272,12 @@ public:
 	// Validate the tensor descriptor, once it's validated any calls to SetXXX() will be ignored
 	bool Validate();
 
+	// Invalidate tensor description, this will reset internal validation to allow calling SetXXX() again.
+	void Invalidate()
+	{
+		bIsValidated = false;
+	}
+
 	// Return filled DML tensor descriptor
 	// NOTE: Call this method only after Validate() is called
 	const DML_TENSOR_DESC* GetDmlDesc() const
@@ -285,23 +310,39 @@ private:
 //
 // DirectML operator base class
 //
-class FOperatorDml
+class FOperatorDml : public IPrepareOperator
 {
 public:
 
 	virtual ~FOperatorDml() = default;
 
-	virtual bool Initialize(IDMLDevice* Device, TArrayView<const NNE::Internal::FTensor> InputTensors, TArrayView<const NNE::Internal::FTensor> OutputTensors, const NNE::FAttributeMap& Attributes) = 0;
+	/**
+	* Initialize is used to only read attributes, not that tensor descriptors can be symbolic
+	*/
+	virtual bool Initialize(TConstArrayView<NNE::FTensorDesc> Inputs, TConstArrayView<NNE::FTensorDesc> Outputs, const NNE::FAttributeMap& Attributes) = 0;
+
+	/**
+	* Evaluate tensor shapes
+	*/
+	virtual int PrepareOutputs(TConstArrayView<NNE::Internal::FTensorRef> InputTensors, TArrayView<NNE::Internal::FTensorRef> OutputTensors) const = 0;
+
+	/**
+	* Once the tensors have concrete shapes we can create the instance of the operator. The concrete shapes are evaluated in PrepareOutputs()
+	*/
+	virtual bool Create(IDMLDevice* Device, TConstArrayView<NNE::Internal::FTensorRef> InputTensors, TConstArrayView<NNE::Internal::FTensorRef> OutputTensors) = 0;
 
+	// Certain operators require tensors to be constant CPU data
 	virtual TConstArrayView<int32> GetConstantCPUInputs() const;
 
+	// Certain operators have different tensor input mappings from ONNX
 	virtual TConstArrayView<int32> GetRemappedInputs() const;
 
+	// Return created DML operator instance
 	IDMLOperator* GetOperator();
 
 protected:
 	
-	bool CreateOperator(IDMLDevice* Device, const DML_OPERATOR_DESC& DmlOpDesc);
+	bool CreateOperator(IDMLDevice* Device, const DML_OPERATOR_DESC& OpDesc);
 
 	TComPtr<IDMLOperator>		DmlOp;
 	Util::FSmallIntArray		ConstantCPUInputs;
diff --git a/Engine/Plugins/Experimental/NNERuntimeRDG/Source/NNERuntimeRDG/Private/Dml/NNEDmlOperatorActivations.cpp b/Engine/Plugins/Experimental/NNERuntimeRDG/Source/NNERuntimeRDG/Private/Dml/NNEDmlOperatorActivations.cpp
--- a/Engine/Plugins/Experimental/NNERuntimeRDG/Source/NNERuntimeRDG/Private/Dml/NNEDmlOperatorActivations.cpp
+++ b/Engine/Plugins/Experimental/NNERuntimeRDG/Source/NNERuntimeRDG/Private/Dml/NNEDmlOperatorActivations.cpp
@@ -18,11 +18,17 @@ class FOperatorDmlActivationUnary : public FOperatorDml
 {
 public:
 
+	//
+	//
+	//
 	static FOperatorDml* Create()
 	{
 		return new FOperatorDmlActivationUnary();
 	}
 
+	//
+	//
+	//
 	static bool Validate(const NNE::FAttributeMap& AttributeMap, TConstArrayView<ENNETensorDataType> InputTypes, TConstArrayView<NNE::FSymbolicTensorShape> InputShapes)
 	{
 		if(InputShapes.Num() != 1)
@@ -33,11 +39,20 @@ public:
 		return CheckElementwiseTensor(InputTypes[0], InputShapes[0]);
 	}
 
+	//
+	//
+	//
 	virtual ~FOperatorDmlActivationUnary() = default;
 
 private:
 
-	FOperatorDmlActivationUnary() : Alpha(0.0f), Beta(0.0f), Gamma(0.0f), Axis(-1) {}
+	//
+	//
+	//
+	FOperatorDmlActivationUnary() : Alpha(0.0f), Beta(0.0f), Gamma(0.0f), Axis(-1) 
+	{
+	}
+	
 	float Alpha;
 	float Beta;
 	float Gamma;
@@ -50,17 +65,16 @@ public:
 	//
 	//
 	//
-	virtual bool Initialize(IDMLDevice* Device, TArrayView<const NNE::Internal::FTensor> InputTensors, TArrayView<const NNE::Internal::FTensor> OutputTensors, const NNE::FAttributeMap& Attributes) override
+	virtual bool Initialize(TConstArrayView<NNE::FTensorDesc> Inputs, TConstArrayView<NNE::FTensorDesc> Outputs, const NNE::FAttributeMap& Attributes) override
 	{
-		const NNE::Internal::FTensor& InputTensor = InputTensors[0];
-		const NNE::Internal::FTensor& OutputTensor = OutputTensors[0];
+		const NNE::FSymbolicTensorShape& InputShape = Inputs[0].GetShape();
 
 		if constexpr (std::is_same_v<DmlActivationOpDescType, DML_ACTIVATION_SOFTMAX1_OPERATOR_DESC> ||
-					  std::is_same_v<DmlActivationOpDescType, DML_ACTIVATION_LOG_SOFTMAX1_OPERATOR_DESC>)
+			std::is_same_v<DmlActivationOpDescType, DML_ACTIVATION_LOG_SOFTMAX1_OPERATOR_DESC>)
 		{
 			Axis = Attributes.GetValueOrDefault(TEXT("axis"), -1);
 		}
-		Axis = HandleNegativeAxis(Axis, InputTensor.GetShape().Rank());
+		Axis = HandleNegativeAxis(Axis, InputShape.Rank());
 
 		Alpha = Attributes.GetValueOrDefault(TEXT("alpha"), Alpha);
 		Beta = Attributes.GetValueOrDefault(TEXT("beta"), Beta);
@@ -72,6 +86,29 @@ public:
 			Treshold = Attributes.GetValueOrDefault(TEXT("lambd"), 0.5f);
 		}
 
+		return true;
+	}
+
+	//
+	//
+	//
+	virtual int PrepareOutputs(TConstArrayView<NNE::Internal::FTensorRef> InputTensors, TArrayView<NNE::Internal::FTensorRef> OutputTensors) const override
+	{
+		check(InputTensors.Num() == 1);
+		check(OutputTensors.Num() == 1);
+		OutputTensors[0]->SetShape(InputTensors[0]->GetShape());
+
+		return 0;
+	}
+
+	//
+	//
+	//
+	virtual bool Create(IDMLDevice* Device, TConstArrayView<NNE::Internal::FTensorRef> InputTensors, TConstArrayView<NNE::Internal::FTensorRef> OutputTensors) override
+	{
+		const NNE::Internal::FTensor& InputTensor = *InputTensors[0];
+		const NNE::Internal::FTensor& OutputTensor = *OutputTensors[0];
+
 		// Initialize tensor descriptor (it's same for both input and output)
 		FTensorDescDml	DmlTensorDesc;
 
@@ -97,6 +134,9 @@ public:
 
 private:
 
+	//
+	//
+	//
 	template<typename OpDesc>
 	void InitDmlOpDesc(OpDesc& Desc, const DML_TENSOR_DESC& TensorDesc)
 	{
@@ -104,6 +144,9 @@ private:
 		Desc.OutputTensor = &TensorDesc;
 	}
 
+	//
+	//
+	//
 	void InitDmlOpDesc(DML_ACTIVATION_LOG_SOFTMAX1_OPERATOR_DESC& Desc, const DML_TENSOR_DESC& TensorDesc)
 	{
 		Desc.InputTensor = &TensorDesc;
@@ -112,6 +155,9 @@ private:
 		Desc.Axes = (UINT*) &Axis;
 	}
 
+	//
+	//
+	//
 	void InitDmlOpDesc(DML_ACTIVATION_SOFTMAX1_OPERATOR_DESC& Desc, const DML_TENSOR_DESC& TensorDesc)
 	{
 		Desc.InputTensor = &TensorDesc;
@@ -120,6 +166,9 @@ private:
 		Desc.Axes = (UINT*) &Axis;
 	}
 
+	//
+	//
+	//
 	void InitDmlOpDesc(DML_ACTIVATION_SOFTPLUS_OPERATOR_DESC& Desc, const DML_TENSOR_DESC& TensorDesc)
 	{
 		Desc.InputTensor = &TensorDesc;
@@ -127,6 +176,9 @@ private:
 		Desc.Steepness = 1.0f;
 	}
 
+	//
+	//
+	//
 	void InitDmlOpDesc(DML_ACTIVATION_SCALED_ELU_OPERATOR_DESC& Desc, const DML_TENSOR_DESC& TensorDesc)
 	{
 		Desc.InputTensor = &TensorDesc;
@@ -135,6 +187,9 @@ private:
 		Desc.Gamma = Gamma;
 	}
 
+	//
+	//
+	//
 	void InitDmlOpDesc(DML_ACTIVATION_ELU_OPERATOR_DESC& Desc, const DML_TENSOR_DESC& TensorDesc)
 	{
 		Desc.InputTensor = &TensorDesc;
@@ -142,6 +197,9 @@ private:
 		Desc.Alpha = Alpha;
 	}
 
+	//
+	//
+	//
 	void InitDmlOpDesc(DML_ACTIVATION_HARD_SIGMOID_OPERATOR_DESC& Desc, const DML_TENSOR_DESC& TensorDesc)
 	{
 		Desc.InputTensor = &TensorDesc;
@@ -150,6 +208,9 @@ private:
 		Desc.Beta = Beta;
 	}
 
+	//
+	//
+	//
 	void InitDmlOpDesc(DML_ACTIVATION_LEAKY_RELU_OPERATOR_DESC& Desc, const DML_TENSOR_DESC& TensorDesc)
 	{
 		Desc.InputTensor = &TensorDesc;
@@ -157,6 +218,9 @@ private:
 		Desc.Alpha = Alpha;
 	}
 
+	//
+	//
+	//
 	void InitDmlOpDesc(DML_ACTIVATION_CELU_OPERATOR_DESC& Desc, const DML_TENSOR_DESC& TensorDesc)
 	{
 		Desc.InputTensor = &TensorDesc;
@@ -164,6 +228,9 @@ private:
 		Desc.Alpha = Alpha;
 	}
 
+	//
+	//
+	//
 	void InitDmlOpDesc(DML_ACTIVATION_SHRINK_OPERATOR_DESC& Desc, const DML_TENSOR_DESC& TensorDesc)
 	{
 		Desc.InputTensor = &TensorDesc;
@@ -186,11 +253,17 @@ class FOperatorDmlActivationBinary : public FOperatorDml
 
 public:
 
+	//
+	//
+	//
 	static FOperatorDml* Create()
 	{
 		return new FOperatorDmlActivationBinary();
 	}
 
+	//
+	//
+	//
 	static bool Validate(const NNE::FAttributeMap& AttributeMap, TConstArrayView<ENNETensorDataType> InputTypes, TConstArrayView<NNE::FSymbolicTensorShape> InputShapes)
 	{
 		if(InputShapes.Num() != 2)
@@ -214,6 +287,9 @@ public:
 
 private:
 
+	//
+	//
+	//
 	FOperatorDmlActivationBinary() = default;
 
 public:
@@ -221,11 +297,31 @@ public:
 	//
 	//
 	//
-	virtual bool Initialize(IDMLDevice* Device, TArrayView<const NNE::Internal::FTensor> InputTensors, TArrayView<const NNE::Internal::FTensor> OutputTensors, const NNE::FAttributeMap& Attributes) override
+	virtual bool Initialize(TConstArrayView<NNE::FTensorDesc> Inputs, TConstArrayView<NNE::FTensorDesc> Outputs, const NNE::FAttributeMap& Attributes) override
+	{
+		return true;
+	}
+
+	//
+	//
+	//
+	virtual int PrepareOutputs(TConstArrayView<NNE::Internal::FTensorRef> InputTensors, TArrayView<NNE::Internal::FTensorRef> OutputTensors) const override
 	{
-		const NNE::Internal::FTensor& InputATensor = InputTensors[0];
-		const NNE::Internal::FTensor& InputBTensor = InputTensors[1];
-		const NNE::Internal::FTensor& OutputTensor = OutputTensors[0];
+		check(InputTensors.Num() == 2);
+		check(OutputTensors.Num() == 1);
+		OutputTensors[0]->SetShape(InputTensors[0]->GetShape());
+
+		return 0;
+	}
+
+	//
+	//
+	//
+	virtual bool Create(IDMLDevice* Device, TConstArrayView<NNE::Internal::FTensorRef> InputTensors, TConstArrayView<NNE::Internal::FTensorRef> OutputTensors)
+	{
+		const NNE::Internal::FTensor& InputATensor = *InputTensors[0];
+		const NNE::Internal::FTensor& InputBTensor = *InputTensors[1];
+		const NNE::Internal::FTensor& OutputTensor = *OutputTensors[0];
 
 		// Initialize tensor descriptors
 		FTensorDescDml	DmlInputATensorDesc;
@@ -233,24 +329,24 @@ public:
 		FTensorDescDml	DmlOutputTensorDesc;
 
 		if (!DmlInputATensorDesc
-				.SetFromTensorBroadcast(InputATensor, OutputTensor.GetShape())
-				.Validate())
+			.SetFromTensorBroadcast(InputATensor, OutputTensor.GetShape())
+			.Validate())
 		{
 			UE_LOG(LogNNE, Error, TEXT("Failed to initialize tensor(s) for DML inference"));
 			return false;
 		}
 
 		if (!DmlInputBTensorDesc
-				.SetFromTensorBroadcast(InputBTensor, OutputTensor.GetShape())
-				.Validate())
+			.SetFromTensorBroadcast(InputBTensor, OutputTensor.GetShape())
+			.Validate())
 		{
 			UE_LOG(LogNNE, Error, TEXT("Failed to initialize tensor(s) for DML inference"));
 			return false;
 		}
 
 		if (!DmlOutputTensorDesc
-				.SetFromTensor(OutputTensor)
-				.Validate())
+			.SetFromTensor(OutputTensor)
+			.Validate())
 		{
 			UE_LOG(LogNNE, Error, TEXT("Failed to initialize tensor(s) for DML inference"));
 			return false;
@@ -266,10 +362,14 @@ public:
 		DmlOpDesc.Desc = &DmlActivationOpDesc;
 
 		return CreateOperator(Device, DmlOpDesc);
+
 	}
 
 private:
 
+	//
+	//
+	//
 	template<typename OpDesc>
 	void InitDmlOpDesc(OpDesc& Desc, DML_TENSOR_DESC& LHSTensorDesc, const DML_TENSOR_DESC& RHSTensorDesc, const DML_TENSOR_DESC& OutputTensorDesc)
 	{
@@ -278,6 +378,9 @@ private:
 		Desc.OutputTensor = &OutputTensorDesc;
 	}
 
+	//
+	//
+	//
 	void InitDmlOpDesc(DML_ACTIVATION_PARAMETERIZED_RELU_OPERATOR_DESC& Desc, const DML_TENSOR_DESC& LHSTensorDesc, const DML_TENSOR_DESC& RHSTensorDesc, const DML_TENSOR_DESC& OutputTensorDesc)
 	{
 		Desc.InputTensor = &LHSTensorDesc;
diff --git a/Engine/Plugins/Experimental/NNERuntimeRDG/Source/NNERuntimeRDG/Private/Dml/NNEDmlOperatorBatchNormalization.cpp b/Engine/Plugins/Experimental/NNERuntimeRDG/Source/NNERuntimeRDG/Private/Dml/NNEDmlOperatorBatchNormalization.cpp
--- a/Engine/Plugins/Experimental/NNERuntimeRDG/Source/NNERuntimeRDG/Private/Dml/NNEDmlOperatorBatchNormalization.cpp
+++ b/Engine/Plugins/Experimental/NNERuntimeRDG/Source/NNERuntimeRDG/Private/Dml/NNEDmlOperatorBatchNormalization.cpp
@@ -6,7 +6,9 @@
 
 namespace UE::NNERuntimeRDG::Private::Dml
 {
-
+//
+//
+//
 class FOperatorDmlBatchNormalization : public FOperatorDml
 {
 	static constexpr float DefaultEpsilon = 1e-5f;
@@ -21,26 +23,36 @@ class FOperatorDmlBatchNormalization : public FOperatorDml
 		Count
 	};
 
+	float Epsilon;
+
 public:
 
+	//
+	//
+	//
 	static FOperatorDml* Create()
 	{
 		return new FOperatorDmlBatchNormalization();
 	}
 
+	//
+	//
+	//
 	static bool Validate(const NNE::FAttributeMap& AttributeMap, TConstArrayView<ENNETensorDataType> InputTypes, TConstArrayView<NNE::FSymbolicTensorShape> InputShapes)
 	{
-		if(InputShapes.Num() != Count)
+		if (InputShapes.Num() != Count)
 		{
 			UE_LOG(LogNNE, Warning, TEXT("DML input tensors must be %d"), Count);
         	return false;
 		}
+		
 		const int32 InputRank = InputShapes[0].Rank();
-		if(InputRank > 8)
+		if (InputRank > 8)
 		{
 			UE_LOG(LogNNE, Warning, TEXT("DML InputTensor rank should be between 1 and 8, it's %d"), InputRank);
         	return false;
 		}
+		
 		const int32 bTrainingMode = AttributeMap.GetValueOrDefault<int32>(TEXT("training_mode"), 0);
 		if (bTrainingMode)
 		{
@@ -50,11 +62,12 @@ public:
 
 		for (int32 Idx = X; Idx < Count; ++Idx)
 		{
-			if(!CheckElementwiseTensor(InputTypes[Idx], InputShapes[Idx]))
+			if (!CheckElementwiseTensor(InputTypes[Idx], InputShapes[Idx]))
 			{
 				return false;
 			}
-			if(!IsEqualOrBroadcastable(InputShapes[0].GetData(), InputShapes[Idx].GetData()))
+			
+			if (!IsEqualOrBroadcastable(InputShapes[0].GetData(), InputShapes[Idx].GetData()))
 			{
 				UE_LOG(LogNNE, Warning, TEXT("DML BatchNormalization other tensors' shapes must be equal or broadcastable to input tensor's"));
 				return false;
@@ -67,9 +80,9 @@ public:
 	//
 	//
 	//
-	virtual bool Initialize(IDMLDevice* Device, TArrayView<const NNE::Internal::FTensor> InputTensors, TArrayView<const NNE::Internal::FTensor> OutputTensors, const NNE::FAttributeMap& Attributes) override
+	virtual bool Initialize(TConstArrayView<NNE::FTensorDesc> Inputs, TConstArrayView<NNE::FTensorDesc> OutputTensors, const NNE::FAttributeMap& Attributes) override
 	{
-		check(InputTensors.Num() == Count);
+		check(Inputs.Num() == Count);
 		check(OutputTensors.Num() >= 1);
 
 		RemappedInputs.Add(X);
@@ -78,17 +91,8 @@ public:
 		RemappedInputs.Add(Scale);
 		RemappedInputs.Add(Bias);
 
-		const NNE::Internal::FTensor& InputTensor = InputTensors[0];
-		const NNE::Internal::FTensor& OutputTensor = OutputTensors[0];
-
-		if (InputTensor.GetShape().Rank() > 8)
-		{
-			UE_LOG(LogNNE, Error, TEXT("InputTensor rank should be between 1 and 8, got:%d"), InputTensor.GetShape().Rank());
-			return false;
-		}
-
 		// Read attributes
-		const float Epsilon = Attributes.GetValueOrDefault<float>(TEXT("epsilon"), DefaultEpsilon);
+		Epsilon = Attributes.GetValueOrDefault<float>(TEXT("epsilon"), DefaultEpsilon);
 		const int32 bTrainingMode = Attributes.GetValueOrDefault<int32>(TEXT("training_mode"), 0);
 
 		if (bTrainingMode || OutputTensors.Num() > 1)
@@ -97,9 +101,33 @@ public:
 			return false;
 		}
 
-		FTensorDescDml DmlInputTensorDescs[Count];
+		return true;
+	}
+
+	//
+	//
+	//
+	virtual int PrepareOutputs(TConstArrayView<NNE::Internal::FTensorRef> InputTensors, TArrayView<NNE::Internal::FTensorRef> OutputTensors) const override
+	{
+		check(InputTensors.Num() == Count);
+		check(OutputTensors.Num() == 1);
+
+		OutputTensors[0]->SetShape(InputTensors[0]->GetShape());
+
+		return 0;
+	};
+
+	//
+	//
+	//
+	virtual bool Create(IDMLDevice* Device, TConstArrayView<NNE::Internal::FTensorRef> InputTensors, TConstArrayView<NNE::Internal::FTensorRef> OutputTensors) override
+	{
+		const NNE::Internal::FTensor& InputTensor = *InputTensors[0];
+		const NNE::Internal::FTensor& OutputTensor = *OutputTensors[0];
+
+		FTensorDescDml DmlInputTensors[Count];
 			
-		if (!DmlInputTensorDescs[X]
+		if (!DmlInputTensors[X]
 				.SetFromTensor(InputTensor)
 				.Validate())
 		{
@@ -107,31 +135,17 @@ public:
 			return false;
 		}
 
-		for (int32 Idx = Scale; Idx < Count; ++Idx)
-		{
-			const NNE::Internal::FTensor& CurrTensor = InputTensors[Idx];
-			FTensorDescDml& DmlCurrTensorDesc = DmlInputTensorDescs[Idx];
+		bool bIsValid = true;
 
-			if (CurrTensor.GetShape().Rank() == 1)
-			{
-				if (!DmlCurrTensorDesc
-						.SetFromTensor1D(CurrTensor, InputTensor.GetShape().Rank())
-						.Validate())
-				{
-					UE_LOG(LogNNE, Error, TEXT("Failed to initialize tensor(s) for DML inference"));
-					return false;
-				}
-			}
-			else
-			{
-				if (!DmlCurrTensorDesc
-						.SetFromTensor(CurrTensor)
-						.Validate())
-				{
-					UE_LOG(LogNNE, Error, TEXT("Failed to initialize tensor(s) for DML inference"));
-					return false;
-				}
-			}
+		bIsValid &= SetDmlTensorDesc(DmlInputTensors[Mean],		*InputTensors[Mean], InputTensor.GetShape().Rank());
+		bIsValid &= SetDmlTensorDesc(DmlInputTensors[Variance],	*InputTensors[Variance], InputTensor.GetShape().Rank());
+		bIsValid &= SetDmlTensorDesc(DmlInputTensors[Scale],	*InputTensors[Scale], InputTensor.GetShape().Rank());
+		bIsValid &= SetDmlTensorDesc(DmlInputTensors[Bias],		*InputTensors[Bias], InputTensor.GetShape().Rank());
+
+		if (!bIsValid)
+		{
+			UE_LOG(LogNNE, Error, TEXT("Failed to initialize tensor(s) for DML inference"));
+			return false;
 		}
 		
 		FTensorDescDml	DmlOutputTensor;
@@ -146,11 +160,11 @@ public:
 
 		DML_BATCH_NORMALIZATION_OPERATOR_DESC	OpDesc{};
 
-		OpDesc.InputTensor = DmlInputTensorDescs[X].GetDmlDesc();
-		OpDesc.MeanTensor = DmlInputTensorDescs[Mean].GetDmlDesc();
-		OpDesc.VarianceTensor = DmlInputTensorDescs[Variance].GetDmlDesc();
-		OpDesc.ScaleTensor = DmlInputTensorDescs[Scale].GetDmlDesc();
-		OpDesc.BiasTensor = DmlInputTensorDescs[Bias].GetDmlDesc();
+		OpDesc.InputTensor = DmlInputTensors[X].GetDmlDesc();
+		OpDesc.MeanTensor = DmlInputTensors[Mean].GetDmlDesc();
+		OpDesc.VarianceTensor = DmlInputTensors[Variance].GetDmlDesc();
+		OpDesc.ScaleTensor = DmlInputTensors[Scale].GetDmlDesc();
+		OpDesc.BiasTensor = DmlInputTensors[Bias].GetDmlDesc();
 		OpDesc.OutputTensor = DmlOutputTensor.GetDmlDesc();
 		OpDesc.Spatial = static_cast<BOOL>(1);
 		OpDesc.Epsilon = Epsilon;
@@ -158,6 +172,37 @@ public:
 
 		return CreateOperator(Device, DML_OPERATOR_DESC{ DML_OPERATOR_BATCH_NORMALIZATION, &OpDesc });
 	}
+
+private:
+
+	//
+	//
+	//
+	inline bool SetDmlTensorDesc(FTensorDescDml& DmlTensorDesc, const NNE::Internal::FTensor& Tensor, int32 InputRank)
+	{
+		if (Tensor.GetShape().Rank() == 1)
+		{
+			if (!DmlTensorDesc
+				.SetFromTensor1D(Tensor, InputRank)
+				.Validate())
+			{
+				UE_LOG(LogNNE, Error, TEXT("Failed to initialize tensor(s) for DML inference"));
+				return false;
+			}
+		}
+		else
+		{
+			if (!DmlTensorDesc
+				.SetFromTensor(Tensor)
+				.Validate())
+			{
+				UE_LOG(LogNNE, Error, TEXT("Failed to initialize tensor(s) for DML inference"));
+				return false;
+			}
+		}
+
+		return true;
+	}
 };
 
 // Register operator on Module startup
diff --git a/Engine/Plugins/Experimental/NNERuntimeRDG/Source/NNERuntimeRDG/Private/Dml/NNEDmlOperatorCast.cpp b/Engine/Plugins/Experimental/NNERuntimeRDG/Source/NNERuntimeRDG/Private/Dml/NNEDmlOperatorCast.cpp
--- a/Engine/Plugins/Experimental/NNERuntimeRDG/Source/NNERuntimeRDG/Private/Dml/NNEDmlOperatorCast.cpp
+++ b/Engine/Plugins/Experimental/NNERuntimeRDG/Source/NNERuntimeRDG/Private/Dml/NNEDmlOperatorCast.cpp
@@ -7,25 +7,33 @@
 
 namespace UE::NNERuntimeRDG::Private::Dml
 {
-
+//
+//
+//
 class FOperatorDmlCast : public FOperatorDml
 {
 public:
 
+	//
+	//
+	//
 	static FOperatorDml* Create()
 	{
 		return new FOperatorDmlCast();
 	}
 
+	//
+	//
+	//
 	static bool Validate(const NNE::FAttributeMap& AttributeMap, TConstArrayView<ENNETensorDataType> InputTypes, TConstArrayView<NNE::FSymbolicTensorShape> InputShapes)
 	{
-		if(InputShapes.Num() != 1)
+		if (InputShapes.Num() != 1)
 		{
 			UE_LOG(LogNNE, Warning, TEXT("There must be only 1 DML input tensor"));
         	return false;
 		}
 
-		if(!CheckGenericTensor(InputTypes[0], InputShapes[0]))
+		if (!CheckGenericTensor(InputTypes[0], InputShapes[0]))
 		{
 			return false;
 		}
@@ -36,25 +44,22 @@ public:
 	//
 	//
 	//
-	virtual bool Initialize(IDMLDevice* Device, TArrayView<const NNE::Internal::FTensor> InputTensors, TArrayView<const NNE::Internal::FTensor> OutputTensors, const NNE::FAttributeMap& Attributes) override
+	virtual bool Initialize(TConstArrayView<NNE::FTensorDesc> Inputs, TConstArrayView<NNE::FTensorDesc> Outputs, const NNE::FAttributeMap& Attributes) override
 	{
-		check(InputTensors.Num() == 1);
-		check(OutputTensors.Num() == 1);
+		check(Inputs.Num() == 1);
+		check(Outputs.Num() == 1);
 		
-		const NNE::Internal::FTensor& InputTensor = InputTensors[0];
-		const NNE::Internal::FTensor& OutputTensor = OutputTensors[0];
+		TConstArrayView<int32> InputShape = Inputs[0].GetShape().GetData();
+		TConstArrayView<int32> OutputShape = Outputs[0].GetShape().GetData();
 
 		ENNETensorDataType To = (ENNETensorDataType) Attributes.GetValue<int32>(TEXT("to"));
 		
-		if (To != OutputTensor.GetDataType())
+		if (To != Outputs[0].GetDataType())
 		{
-			UE_LOG(LogNNE, Error, TEXT("Cast should output a tensor of type %d but was of type %d."), int(To), int(OutputTensor.GetDataType()));
+			UE_LOG(LogNNE, Error, TEXT("Cast should output a tensor of type %d but was of type %d."), int(To), int(Outputs[0].GetDataType()));
 			return false;
 		}
 
-		TConstArrayView<uint32> InputShape = InputTensor.GetShape().GetData();
-		TConstArrayView<uint32> OutputShape = OutputTensor.GetShape().GetData();
-
 		if (InputShape.Num() != OutputShape.Num())
 		{
 			UE_LOG(LogNNE, Error, TEXT("Cast input and output shapes need to have a same rank"));
@@ -70,6 +75,29 @@ public:
 			}
 		}
 
+		return true;
+	}
+
+	//
+	//
+	//
+	virtual int PrepareOutputs(TConstArrayView<NNE::Internal::FTensorRef> InputTensors, TArrayView<NNE::Internal::FTensorRef> OutputTensors) const override
+	{
+		check(InputTensors.Num() == 1);
+		check(OutputTensors.Num() == 1);
+		OutputTensors[0]->SetShape(InputTensors[0]->GetShape());
+
+		return 0;
+	};
+
+	//
+	//
+	//
+	virtual bool Create(IDMLDevice* Device, TConstArrayView<NNE::Internal::FTensorRef> InputTensors, TConstArrayView<NNE::Internal::FTensorRef> OutputTensors) override
+	{
+		const NNE::Internal::FTensor& InputTensor = *InputTensors[0];
+		const NNE::Internal::FTensor& OutputTensor = *OutputTensors[0];
+
 		// Initialize tensor descriptors
 		FTensorDescDml DmlInputTensorDesc;
 		
diff --git a/Engine/Plugins/Experimental/NNERuntimeRDG/Source/NNERuntimeRDG/Private/Dml/NNEDmlOperatorConcat.cpp b/Engine/Plugins/Experimental/NNERuntimeRDG/Source/NNERuntimeRDG/Private/Dml/NNEDmlOperatorConcat.cpp
--- a/Engine/Plugins/Experimental/NNERuntimeRDG/Source/NNERuntimeRDG/Private/Dml/NNEDmlOperatorConcat.cpp
+++ b/Engine/Plugins/Experimental/NNERuntimeRDG/Source/NNERuntimeRDG/Private/Dml/NNEDmlOperatorConcat.cpp
@@ -15,63 +15,61 @@ namespace UE::NNERuntimeRDG::Private::Dml
  */
 class FOperatorDmlConcat : public FOperatorDml
 {
-	
+	int32 Axis;
 
 public:
 
+	//
+	//
+	//
 	static FOperatorDml* Create()
 	{
 		return new FOperatorDmlConcat();
 	}
 
+	//
+	//
+	//
 	static bool Validate(const NNE::FAttributeMap& AttributeMap, TConstArrayView<ENNETensorDataType> InputTypes, TConstArrayView<NNE::FSymbolicTensorShape> InputShapes)
 	{
-		if(InputShapes.Num() == 0)
+		if (InputShapes.Num() == 0)
 		{
 			UE_LOG(LogNNE, Warning, TEXT("DML no input tensors"));
-        	return false;
+			return false;
 		}
 
-		const int32 InputDims = InputShapes[0].Rank();
-		auto NormalizeAxis = [InputDims] (int32 Axis) -> uint32
-		{
-			if (Axis < 0)
-			{
-				Axis += InputDims;
-			}
-			return (uint32) Axis;
-		};
-
-		uint32 Axis = NormalizeAxis(AttributeMap.GetValue<int32>(TEXT("axis")));
-		if(Axis >= (uint32) InputDims)
+		const int32 InputRank = InputShapes[0].Rank();
+		
+		int32 Axis = HandleNegativeAxis(AttributeMap.GetValue<int32>(TEXT("axis")), InputRank);
+		if (Axis >= InputRank)
 		{
 			UE_LOG(LogNNE, Warning, TEXT("DML invalid axis: %d"), Axis);
-        	return false;
+			return false;
 		}
 
-		if(!CheckGenericTensor(InputTypes[0], InputShapes[0]))
+		if (!CheckGenericTensor(InputTypes[0], InputShapes[0]))
 		{
 			return false;
 		}
 		
 		for (int32 Idx = 1; Idx < InputShapes.Num(); ++Idx)
 		{
-			if(!CheckGenericTensor(InputTypes[Idx], InputShapes[Idx]))
+			if (!CheckGenericTensor(InputTypes[Idx], InputShapes[Idx]))
 			{
 				return false;
 			}
 
-			if(InputShapes[Idx].Rank() != InputDims)
+			if (InputShapes[Idx].Rank() != InputRank)
 			{
 				UE_LOG(LogNNE, Warning, TEXT("DML concat rank mismatch for tensor %d"), Idx);
 				return false;
 			}
 
-			for(int32 Dim = 0; Dim < InputShapes[Idx].Rank(); ++Dim)
+			for (int32 Dim = 0; Dim < InputShapes[Idx].Rank(); ++Dim)
 			{
-				if(Dim != Axis)
+				if (Dim != Axis)
 				{
-					if(InputShapes[Idx].GetData()[Dim] != InputShapes[0].GetData()[Dim])
+					if (InputShapes[Idx].GetData()[Dim] != InputShapes[0].GetData()[Dim])
 					{
 						UE_LOG(LogNNE, Warning, TEXT("DML concat dimension mismatch for tensor %d"), Idx);
 						return false;
@@ -86,39 +84,75 @@ public:
 	//
 	//
 	//
-	virtual bool Initialize(IDMLDevice* Device, TArrayView<const NNE::Internal::FTensor> InputTensors, TArrayView<const NNE::Internal::FTensor> OutputTensors, const NNE::FAttributeMap& Attributes) override
+	virtual bool Initialize(TConstArrayView<NNE::FTensorDesc> Inputs, TConstArrayView<NNE::FTensorDesc> Outputs, const NNE::FAttributeMap& Attributes) override
 	{
-		check(InputTensors.Num() > 0);
-		check(OutputTensors.Num() == 1);
+		check(Inputs.Num() >= 1);
+		check(Outputs.Num() == 1);
+
+		const int32 InputRank = Inputs[0].GetShape().Rank();
 
-		int32 InputDims = InputTensors[0].GetShape().Rank();
-		auto NormalizeAxis = [InputDims] (int32 Axis) -> uint32
+		Axis = Attributes.GetValue<int32>(TEXT("axis"));
+
+		if (Axis < -InputRank || Axis >(InputRank - 1))
 		{
-			if (Axis < 0)
+			UE_LOG(LogNNE, Warning, TEXT("Axis should be in range [-r,r-1] however it is %d while inputs have rank %d."), Axis, InputRank);
+			return false;
+		}
+
+		Axis = HandleNegativeAxis(Axis, InputRank);
+		check(Axis < InputRank);
+
+		return true;
+	}
+
+	//
+	//
+	//
+	virtual int PrepareOutputs(TConstArrayView<NNE::Internal::FTensorRef> InputTensors, TArrayView<NNE::Internal::FTensorRef> OutputTensors) const override
+	{
+		TArray<uint32> OutputShape(InputTensors[0]->GetShape().GetData());
+		
+		const int32 AxisIndex = Axis >= 0 ? Axis : InputTensors[0]->GetShape().Rank() - Axis;
+
+		for (int32 i = 1; i < InputTensors.Num(); ++i)
+		{
+			OutputShape[AxisIndex] += InputTensors[i]->GetShape().GetData()[AxisIndex];
+
+			for (int32 r = 0; r < OutputShape.Num(); ++r)
 			{
-				Axis += InputDims;
+				if (r != AxisIndex && (OutputShape[r] != InputTensors[i]->GetShape().GetData()[r]))
+				{
+					UE_LOG(LogNNE, Warning, TEXT("Concat: all input tensors should have the same shape except on the concatenation axis"));
+					return false;
+				}
 			}
-			return (uint32) Axis;
-		};
+		}
 
-		uint32 Axis = NormalizeAxis(Attributes.GetValue<int32>(TEXT("axis")));
-		check(Axis < (uint32) InputDims);
+		OutputTensors[0]->SetShape(NNE::FTensorShape::Make(OutputShape));
+
+		return 0;
+	}
 
+	//
+	//
+	//
+	virtual bool Create(IDMLDevice* Device, TConstArrayView<NNE::Internal::FTensorRef> InputTensors, TConstArrayView<NNE::Internal::FTensorRef> OutputTensors) override
+	{
 		TArray<FTensorDescDml> InputDescs;
 		InputDescs.SetNum(InputTensors.Num());
 		
 		TArray<DML_TENSOR_DESC> DmlInputDescs;
 		DmlInputDescs.SetNumUninitialized(InputTensors.Num());
 		
-		for (TConstEnumerateRef<NNE::Internal::FTensor> It : EnumerateRange(InputTensors))
+		for (TConstEnumerateRef<NNE::Internal::FTensorRef> It : EnumerateRange(InputTensors))
 		{
-			const NNE::Internal::FTensor& Tensor = *It;
+			const NNE::Internal::FTensorRef& Tensor = *It;
 			int Idx = It.GetIndex();
 			
-			if (Algo::Find(Tensor.GetShape().GetData(), 0) == nullptr)
+			if (Algo::Find(Tensor->GetShape().GetData(), 0) == nullptr)
 			{
 				if (!InputDescs[Idx]
-						.SetFromTensor(Tensor)
+						.SetFromTensor(*Tensor)
 						.Validate())
 				{
 					UE_LOG(LogNNE, Error, TEXT("Failed to initialize Concat input for DML inference"));
@@ -128,13 +162,12 @@ public:
 				DmlInputDescs[Idx] = *InputDescs[Idx].GetDmlDesc();
 			}
 		}
-		
 		check(InputDescs.Num() > 0);
 
 		FTensorDescDml OutputTensorDesc;
 
 		if (!OutputTensorDesc
-				.SetFromTensor(OutputTensors[0])
+				.SetFromTensor(*OutputTensors[0])
 				.Validate())
 		{
 			UE_LOG(LogNNE, Error, TEXT("Failed to initialize Concat output for DML inference"));
diff --git a/Engine/Plugins/Experimental/NNERuntimeRDG/Source/NNERuntimeRDG/Private/Dml/NNEDmlOperatorConvolution.cpp b/Engine/Plugins/Experimental/NNERuntimeRDG/Source/NNERuntimeRDG/Private/Dml/NNEDmlOperatorConvolution.cpp
--- a/Engine/Plugins/Experimental/NNERuntimeRDG/Source/NNERuntimeRDG/Private/Dml/NNEDmlOperatorConvolution.cpp
+++ b/Engine/Plugins/Experimental/NNERuntimeRDG/Source/NNERuntimeRDG/Private/Dml/NNEDmlOperatorConvolution.cpp
@@ -13,98 +13,37 @@ namespace UE::NNERuntimeRDG::Private::Dml
 template <DML_CONVOLUTION_DIRECTION Direction>
 class FOperatorDmlConv : public FOperatorDml
 {
-	using FSmallArray = TArray<uint32, TInlineAllocator<NcdhwSpatialDimensionCount>>;
-	using FIntArray = TArray<int32>;
-
-	struct FConvArgs
+	//
+	//
+	//
+	class FConvArgs : public FKernelArgs
 	{
-		EAutoPad				AutoPad;
-		Util::FSmallUIntArray	StartPadding;
-		Util::FSmallUIntArray	EndPadding;
-		FIntArray				OutPadding;
-		FIntArray				Dilations;
-		FIntArray				Strides;
-		FSmallArray				OutputShape;
-		uint32					NumDimensions;
-		FSmallArray				WindowSize;
-		int32					Group;
+		int32			Group;
 
+	public:
+
+		//
+		//
+		//
 		FConvArgs() = default;
 
 		//
 		//
 		//
-		bool Init(const NNE::FTensorShape& InputShape, const NNE::FTensorShape& FilterShape, const NNE::FAttributeMap& Attributes)
+		bool Init(const NNE::FTensorDesc& Input, const NNE::FTensorDesc& Filter, const NNE::FAttributeMap& Attributes)
 		{
-			check(InputShape.Rank() > NonspatialDimensionCount);
-			check(FilterShape.Rank() == InputShape.Rank());
-			
-			NumDimensions = InputShape.Rank() - NonspatialDimensionCount;
-
-			const FNNEAttributeValue*	AttrStrides = Attributes.GetAttributeValue(TEXT("strides"));
-
-			if (AttrStrides)
-			{
-				Strides = AttrStrides->GetValue<FIntArray>();
-			}
-			else
-			{
-				Strides.Init(1, NumDimensions);
-			}
-
-			check(Strides.Num() == 0 || Strides.Num() == FilterShape.Rank() - NonspatialDimensionCount);
-
-			const FNNEAttributeValue*	AttrDilations = Attributes.GetAttributeValue(TEXT("dilations"));
-
-			if (AttrDilations)
-			{
-				Dilations = AttrDilations->GetValue<FIntArray>();
-			}
-			else
+			if (!FKernelArgs::Init(Attributes, Input.GetShape().Rank(), /*bIsGlobalKernel*/ false, /*bIsTransposed*/ Direction != DML_CONVOLUTION_DIRECTION_FORWARD))
 			{
-				Dilations.Init(1, NumDimensions);
-			}
-
-			check(Dilations.Num() == 0 || Dilations.Num() == FilterShape.Rank() - NonspatialDimensionCount);
-			
-			for (int32 Dim = FilterShape.Rank() - NumDimensions; Dim < FilterShape.Rank(); ++Dim)
-			{
-				WindowSize.Add(FilterShape.GetData()[Dim]);
+				return false;
 			}
 
-			if (Direction == DML_CONVOLUTION_DIRECTION_FORWARD)
-			{
-				OutPadding.Init(0, NumDimensions);
-			}
-			else
-			{
-				const FNNEAttributeValue* AttrOutPadding = Attributes.GetAttributeValue(TEXT("output_padding"));
-
-				if (AttrOutPadding)
-				{
-					OutPadding = AttrOutPadding->GetValue<TArray<int32>>();
-				}
-				else
-				{
-					OutPadding.Init(0, NumDimensions);
-				}
-			}
+			check(Filter.GetShape().Rank() == Input.GetShape().Rank());
+			check(Strides.Num() == 0 || Strides.Num() == Filter.GetShape().Rank() - NonspatialDimensionCount);
+			check(Dilations.Num() == 0 || Dilations.Num() == Filter.GetShape().Rank() - NonspatialDimensionCount);
 
 			Group = Attributes.GetValueOrDefault<int32>(TEXT("group"), 1);
 
-			ComputeStartEndPaddings(
-				InputShape.GetData(),
-				Attributes, 
-				StartPadding, 
-				EndPadding,
-				ConvolutionPadding(InputShape.GetData())
-			);
-
-			if (Direction == DML_CONVOLUTION_DIRECTION_FORWARD)
-			{
-				SetOutputShape(InputShape.GetData(), FilterShape.GetData());
-			}
-			else
+			if (Direction != DML_CONVOLUTION_DIRECTION_FORWARD)
 			{
 				const FNNEAttributeValue* AttrOutShape = Attributes.GetAttributeValue(TEXT("output_shape"));
 
@@ -117,10 +56,6 @@ class FOperatorDmlConv : public FOperatorDml
 						OutputShape.Add(uint32(Value));
 					}
 				}
-				else
-				{
-					SetOutputShape(InputShape.GetData(), FilterShape.GetData());
-				}
 			}
 
 			return true;
@@ -129,84 +64,88 @@ class FOperatorDmlConv : public FOperatorDml
 		//
 		//
 		//
-		Util::FSmallUIntArray ConvolutionPadding(TConstArrayView<uint32> InputShape)
+		void Evaluate(TConstArrayView<uint32> InputShape, TConstArrayView<uint32> FilterShape)
 		{
-			const uint32 DimOffset = NonspatialDimensionCount;
+			// NOTE: To compute paddings we need WindowSize
+			for (int32 Dim = FilterShape.Num() - NumDimensions; Dim < FilterShape.Num(); ++Dim)
+			{
+				WindowSize.Add(FilterShape[Dim]);
+			}
 
-			if(Direction == DML_CONVOLUTION_DIRECTION_FORWARD)
+			FKernelArgs::Evaluate(InputShape, ConvolutionPadding(InputShape));
+
+			if (Direction == DML_CONVOLUTION_DIRECTION_FORWARD)
 			{
-				return KernelPadding(
-					InputShape, WindowSize,
-					MakeArrayView((uint32*) Dilations.GetData(), Dilations.Num()), MakeArrayView((uint32*) Strides.GetData(), Strides.Num())
-				);
+				OutputShape[1] = FilterShape[0];
 			}
-			// Deconvolution
 			else
 			{
-				Util::FSmallUIntArray Padding;
-				Padding.SetNumUninitialized(NumDimensions);
-				
-				for (uint32 Dim = 0; Dim < NumDimensions; ++Dim)
+				if (!bHasOutputShape)
 				{
-					Padding[Dim] = (InputShape[Dim + DimOffset] - 1) * Dilations[Dim] - Strides[Dim] + OutPadding[Dim] + 1;
+					OutputShape[1] = FilterShape[1] * Group;
 				}
-
-				return Padding;
 			}
 		}
 
 		//
 		//
 		//
-		void SetOutputShape(TConstArrayView<uint32> InputShape, TConstArrayView<uint32> FilterShape)
+		TConstArrayView<uint32> GetOutPadding() const
 		{
-			const uint32 DimOffset = NonspatialDimensionCount;
-
-			OutputShape.SetNumUninitialized(InputShape.Num());
-
-			if (Direction == DML_CONVOLUTION_DIRECTION_FORWARD)
-			{
-				OutputShape[0] = InputShape[0];
-				OutputShape[1] = FilterShape[0];
+			return MakeArrayView((const uint32*) OutPadding.GetData(), OutPadding.Num());
+		}
 
-				for (uint32 Dim = 0; Dim < NumDimensions; ++Dim)
-				{
-					uint32 InputLen = InputShape[Dim + DimOffset];
-					uint32 PaddedLen = InputLen + StartPadding[Dim] + EndPadding[Dim];
-					uint32 KernelLen = 1 + (WindowSize[Dim] - 1) * Dilations[Dim];
+		int32 GetGroup() const
+		{
+			return Group;
+		}
 
-					checkf(KernelLen <= PaddedLen, TEXT("KernelLen must < PaddedLen"));
-					checkf(Strides[Dim] != 0, TEXT("Strides must be != 0"));
+	private:
 
-					uint32 StridableOutLen = PaddedLen - KernelLen;
-					uint32 OutLen = 1 + (StridableOutLen / Strides[Dim]);
+		//
+		//
+		//
+		Util::FSmallUIntArray ConvolutionPadding(TConstArrayView<uint32> InputShape)
+		{
+			const uint32 DimOffset = NonspatialDimensionCount;
 
-					OutputShape[Dim + DimOffset] = OutLen;
-				}
+			if (Direction == DML_CONVOLUTION_DIRECTION_FORWARD)
+			{
+				return KernelPadding(
+					InputShape, WindowSize,
+					MakeArrayView((uint32*) Dilations.GetData(), Dilations.Num()), MakeArrayView((uint32*) Strides.GetData(), Strides.Num())
+				);
 			}
 			else
 			{
-				OutputShape[0] = InputShape[0];
-				OutputShape[1] = FilterShape[1] * Group;
-
+				Util::FSmallUIntArray Padding;
+				Padding.SetNumUninitialized(NumDimensions);
+				
 				for (uint32 Dim = 0; Dim < NumDimensions; ++Dim)
 				{
-					uint32 Padding = StartPadding[Dim] + EndPadding[Dim];
-					uint32 KernelLen = 1 + (WindowSize[Dim] - 1) * Dilations[Dim];
-
-					OutputShape[Dim + DimOffset] = (InputShape[Dim + DimOffset] - 1) * Strides[Dim] + KernelLen + OutPadding[Dim] - Padding;
+					Padding[Dim] = (InputShape[Dim + DimOffset] - 1) * Dilations[Dim] - Strides[Dim] + OutPadding[Dim] + 1;
 				}
+
+				return Padding;
 			}
 		}
 	};
 
+	mutable FConvArgs	Args;
+
 public:
 
+	//
+	//
+	//
 	static FOperatorDml* Create()
 	{
 		return new FOperatorDmlConv();
 	}
 
+	//
+	//
+	//
 	static bool Validate(const NNE::FAttributeMap& AttributeMap, TConstArrayView<ENNETensorDataType> InputTypes, TConstArrayView<NNE::FSymbolicTensorShape> InputShapes)
 	{
 		//TODO
@@ -216,24 +155,40 @@ public:
 	//
 	//
 	//
-	virtual bool Initialize(IDMLDevice* Device, TArrayView<const NNE::Internal::FTensor> InputTensors, TArrayView<const NNE::Internal::FTensor> OutputTensors, const NNE::FAttributeMap& Attributes) override
+	virtual bool Initialize(TConstArrayView<NNE::FTensorDesc> Inputs, TConstArrayView<NNE::FTensorDesc> Outputs, const NNE::FAttributeMap& Attributes) override
 	{
-		const NNE::Internal::FTensor& InputTensor = InputTensors[0];
-		const NNE::Internal::FTensor& FilterTensor = InputTensors[1];
-		
-		FConvArgs	Args;
-		
-		if (!Args.Init(InputTensor.GetShape(), FilterTensor.GetShape(), Attributes))
+		const NNE::FTensorDesc& InputTensor = Inputs[0];
+		const NNE::FTensorDesc& FilterTensor = Inputs[1];
+
+		if (!Args.Init(InputTensor, FilterTensor, Attributes))
 		{
 			return false;
 		}
 
-		NNE::Internal::FTensor OutputTensor = OutputTensors[0];
+		return true;
+	}
 
-		if (Direction == DML_CONVOLUTION_DIRECTION_FORWARD)
-		{
-			OutputTensor.SetShape(NNE::FTensorShape::Make(Args.OutputShape));
-		}
+	//
+	//
+	//
+	virtual int PrepareOutputs(TConstArrayView<NNE::Internal::FTensorRef> InputTensors, TArrayView<NNE::Internal::FTensorRef> OutputTensors) const override
+	{
+		NNE::Internal::FTensor& OutputTensor = *OutputTensors[0];
+
+		Args.Evaluate(InputTensors[0]->GetShape().GetData(), InputTensors[1]->GetShape().GetData());
+		OutputTensor.SetShape(NNE::FTensorShape::Make(Args.GetOutputShape()));
+
+		return 0;
+	}
+
+	//
+	//
+	//
+	virtual bool Create(IDMLDevice* Device, TConstArrayView<NNE::Internal::FTensorRef> InputTensors, TConstArrayView<NNE::Internal::FTensorRef> OutputTensors) override
+	{
+		const NNE::Internal::FTensor& InputTensor = *InputTensors[0];
+		const NNE::Internal::FTensor& FilterTensor = *InputTensors[1];
+		const NNE::Internal::FTensor& OutputTensor = *OutputTensors[0];
 
 		// Initialize tensor descriptors
 		FTensorDescDml	DmlInputTensorDesc;
@@ -261,7 +216,7 @@ public:
 
 		if (InputTensors.Num() > 2)
 		{
-			const NNE::Internal::FTensor& BiasTensor = InputTensors[2];
+			const NNE::Internal::FTensor& BiasTensor = *InputTensors[2];
 
 			if (!DmlBiasTensorDesc
 					.SetTensorRank(3, 5)
@@ -290,13 +245,13 @@ public:
 		DmlConvOpDesc.OutputTensor = DmlOutputTensorDesc.GetDmlDesc();
 		DmlConvOpDesc.Mode = DML_CONVOLUTION_MODE_CROSS_CORRELATION;
 		DmlConvOpDesc.Direction = Direction;
-		DmlConvOpDesc.DimensionCount = Args.NumDimensions;
-		DmlConvOpDesc.Strides = (uint32*) Args.Strides.GetData();
-		DmlConvOpDesc.Dilations = (uint32*) Args.Dilations.GetData();
-		DmlConvOpDesc.StartPadding = Args.StartPadding.GetData();
-		DmlConvOpDesc.EndPadding = Args.EndPadding.GetData();
-		DmlConvOpDesc.OutputPadding = (uint32*) Args.OutPadding.GetData();
-		DmlConvOpDesc.GroupCount = Args.Group;
+		DmlConvOpDesc.DimensionCount = Args.GetNumDimensions();
+		DmlConvOpDesc.Strides = Args.GetStrides().GetData();
+		DmlConvOpDesc.Dilations = Args.GetDilations().GetData();
+		DmlConvOpDesc.StartPadding = Args.GetStartPadding().GetData();
+		DmlConvOpDesc.EndPadding = Args.GetEndPadding().GetData();
+		DmlConvOpDesc.OutputPadding = Args.GetOutPadding().GetData();
+		DmlConvOpDesc.GroupCount = Args.GetGroup();
 
 		DML_OPERATOR_DESC DmlOpDesc{};
 
@@ -307,16 +262,25 @@ public:
 	}
 };
 
+//
+//
+//
 void RegisterConvOperator()
 {
 	FOperatorRegistryDml::Get()->OpAdd(TEXT("Conv"), FOperatorDmlConv<DML_CONVOLUTION_DIRECTION_FORWARD>::Create);
 }
 
+//
+//
+//
 void RegisterConvTransposeOperator()
 {
 	FOperatorRegistryDml::Get()->OpAdd(TEXT("ConvTranspose"), FOperatorDmlConv<DML_CONVOLUTION_DIRECTION_BACKWARD>::Create);
 }
 
+//
+//
+//
 struct FDmlOperatorConvRegistrator
 {
 	FDmlOperatorConvRegistrator()
@@ -326,6 +290,9 @@ struct FDmlOperatorConvRegistrator
 	}
 };
 
+//
+//
+//
 static FDmlOperatorConvRegistrator RegisterDmlOperatorConv;
 
 } // namespace UE::NNERuntimeRDG::Private::Dml
diff --git a/Engine/Plugins/Experimental/NNERuntimeRDG/Source/NNERuntimeRDG/Private/Dml/NNEDmlOperatorCumSum.cpp b/Engine/Plugins/Experimental/NNERuntimeRDG/Source/NNERuntimeRDG/Private/Dml/NNEDmlOperatorCumSum.cpp
--- a/Engine/Plugins/Experimental/NNERuntimeRDG/Source/NNERuntimeRDG/Private/Dml/NNEDmlOperatorCumSum.cpp
+++ b/Engine/Plugins/Experimental/NNERuntimeRDG/Source/NNERuntimeRDG/Private/Dml/NNEDmlOperatorCumSum.cpp
@@ -8,16 +8,28 @@
 namespace UE::NNERuntimeRDG::Private::Dml
 {
 
+//
+//
+//
 class FOperatorDmlCumSum : public FOperatorDml
 {
+	DML_AXIS_DIRECTION	AxisDirection;
+	int32				HasExclusiveSum;
+	mutable int32		Axis;
 
 public:
 
+	//
+	//
+	//
 	static FOperatorDml* Create()
 	{
 		return new FOperatorDmlCumSum();
 	}
 
+	//
+	//
+	//
 	static bool Validate(const NNE::FAttributeMap& AttributeMap, TConstArrayView<ENNETensorDataType> InputTypes, TConstArrayView<NNE::FSymbolicTensorShape> InputShapes)
 	{
 		//TODO
@@ -27,35 +39,45 @@ public:
 	//
 	//
 	//
-	virtual bool Initialize(IDMLDevice* Device, TArrayView<const NNE::Internal::FTensor> InputTensors, TArrayView<const NNE::Internal::FTensor> OutputTensors, const NNE::FAttributeMap& Attributes) override
+	virtual bool Initialize(TConstArrayView<NNE::FTensorDesc> Inputs, TConstArrayView<NNE::FTensorDesc> Outputs, const NNE::FAttributeMap& Attributes) override
 	{
-		check(InputTensors.Num() == 2);
-		check(OutputTensors.Num() == 1);
+		check(Inputs.Num() == 2);
+		check(Outputs.Num() == 1);
 
 		ConstantCPUInputs.Add(1);
 
-		UINT Axis; 
+		AxisDirection = (DML_AXIS_DIRECTION) Attributes.GetValueOrDefault<int32>(TEXT("reverse"), 0);
+		HasExclusiveSum = Attributes.GetValueOrDefault<int32>(TEXT("exclusive"), 0);
 
-        // axis tensor must be constant
-        check(InputTensors[1].HasPreparedData());
+		return true;
+	}
+
+	//
+	//
+	//
+	virtual int PrepareOutputs(TConstArrayView<NNE::Internal::FTensorRef> InputTensors, TArrayView<NNE::Internal::FTensorRef> OutputTensors) const override
+	{
+		// axis tensor must be constant
+		check(InputTensors[1]->HasPreparedData());
 
-		switch(InputTensors[1].GetDataType())
-        {
-        case ENNETensorDataType::Int32:
+		switch(InputTensors[1]->GetDataType())
+		{
+			case ENNETensorDataType::Int32:
 			{
-				TConstArrayView<int32> TensorContent = InputTensors[1].GetPreparedData<int32>();
+				TConstArrayView<int32> TensorContent = InputTensors[1]->GetPreparedData<int32>();
 				if (TensorContent.Num() != 1)
 				{
 					UE_LOG(LogNNE, Error, TEXT("axis tensor should be 0-D"));
 					return false;
 				}
-				Axis = (UINT) HandleNegativeAxis(TensorContent[0], InputTensors[0].GetShape().Rank());
+				
+				Axis = HandleNegativeAxis(TensorContent[0], InputTensors[0]->GetShape().Rank());
 			}
-            break;
+			break;
 
-        case ENNETensorDataType::Int64:
+			case ENNETensorDataType::Int64:
 			{
-				TConstArrayView<int64> TensorContent = InputTensors[1].GetPreparedData<int64>();
+				TConstArrayView<int64> TensorContent = InputTensors[1]->GetPreparedData<int64>();
 				if (TensorContent.Num() != 1)
 				{
 					UE_LOG(LogNNE, Error, TEXT("axis tensor should be 0-D"));
@@ -67,17 +89,28 @@ public:
 					UE_LOG(LogNNE, Error, TEXT("axis number is overflowing"));
 					return false;
 				}
-				Axis = (UINT) HandleNegativeAxis(TensorContent[0], InputTensors[0].GetShape().Rank());
+				
+				Axis = HandleNegativeAxis(TensorContent[0], InputTensors[0]->GetShape().Rank());
 			}
-            break;
+			break;
 
-        default:
-            UE_LOG(LogNNE, Error, TEXT("axis tensor has invalid data type"));
-			return false;
-        };
+			default:
+				UE_LOG(LogNNE, Error, TEXT("axis tensor has invalid data type"));
+				return -1;
+		};
+
+		OutputTensors[0]->SetShape(InputTensors[0]->GetShape());
 
-		const NNE::Internal::FTensor& InputTensor = InputTensors[0];
-		const NNE::Internal::FTensor& OutputTensor = OutputTensors[0];
+		return 0;
+	}
+
+	//
+	//
+	//
+	virtual bool Create(IDMLDevice* Device, TConstArrayView<NNE::Internal::FTensorRef> InputTensors, TConstArrayView<NNE::Internal::FTensorRef> OutputTensors) override
+	{
+		const NNE::Internal::FTensor& InputTensor = *InputTensors[0];
+		const NNE::Internal::FTensor& OutputTensor = *OutputTensors[0];
 
 		FTensorDescDml	DmlInputTensorDesc;
 		FTensorDescDml	DmlOutputTensorDesc;
@@ -103,8 +136,8 @@ public:
 		DmlCumSumOpDesc.InputTensor = DmlInputTensorDesc.GetDmlDesc();
 		DmlCumSumOpDesc.OutputTensor = DmlOutputTensorDesc.GetDmlDesc();
 		DmlCumSumOpDesc.Axis = Axis;
-		DmlCumSumOpDesc.AxisDirection = (DML_AXIS_DIRECTION) Attributes.GetValueOrDefault<int32>(TEXT("reverse"), 0);
-		DmlCumSumOpDesc.HasExclusiveSum = (bool) Attributes.GetValueOrDefault<int32>(TEXT("exclusive"), 0);
+		DmlCumSumOpDesc.AxisDirection = AxisDirection;
+		DmlCumSumOpDesc.HasExclusiveSum = HasExclusiveSum;
 
 		return CreateOperator(Device, DML_OPERATOR_DESC{ DML_OPERATOR_CUMULATIVE_SUMMATION, &DmlCumSumOpDesc });
 	}
diff --git a/Engine/Plugins/Experimental/NNERuntimeRDG/Source/NNERuntimeRDG/Private/Dml/NNEDmlOperatorDepthToSpace.cpp b/Engine/Plugins/Experimental/NNERuntimeRDG/Source/NNERuntimeRDG/Private/Dml/NNEDmlOperatorDepthToSpace.cpp
--- a/Engine/Plugins/Experimental/NNERuntimeRDG/Source/NNERuntimeRDG/Private/Dml/NNEDmlOperatorDepthToSpace.cpp
+++ b/Engine/Plugins/Experimental/NNERuntimeRDG/Source/NNERuntimeRDG/Private/Dml/NNEDmlOperatorDepthToSpace.cpp
@@ -8,9 +8,14 @@
 namespace UE::NNERuntimeRDG::Private::Dml
 {
 
+//
+//
+//
 class FOperatorDmlDepthToSpace : public FOperatorDml
 {
-
+	//
+	//
+	//
 	static DML_DEPTH_SPACE_ORDER SpaceOrderFromModeString(FStringView StringVal)
 	{
 		if (FCString::Stricmp(StringVal.GetData(), TEXT("CRD")) == 0)
@@ -23,13 +28,28 @@ class FOperatorDmlDepthToSpace : public FOperatorDml
 		}
 	}
 
+	enum InputDims
+	{
+		N, C, H, W,
+		DIM_COUNT
+	};
+
+	DML_DEPTH_SPACE_ORDER	Order;
+	int32					BlockSize;
+
 public:
 
+	//
+	//
+	//
 	static FOperatorDml* Create()
 	{
 		return new FOperatorDmlDepthToSpace();
 	}
 
+	//
+	//
+	//
 	static bool Validate(const NNE::FAttributeMap& AttributeMap, TConstArrayView<ENNETensorDataType> InputTypes, TConstArrayView<NNE::FSymbolicTensorShape> InputShapes)
 	{
 		//TODO
@@ -39,12 +59,10 @@ public:
 	//
 	//
 	//
-	virtual bool Initialize(IDMLDevice* Device, TArrayView<const NNE::Internal::FTensor> InputTensors, TArrayView<const NNE::Internal::FTensor> OutputTensors, const NNE::FAttributeMap& Attributes) override
+	virtual bool Initialize(TConstArrayView<NNE::FTensorDesc> Inputs, TConstArrayView<NNE::FTensorDesc> Outputs, const NNE::FAttributeMap& Attributes) override
 	{
-		check(InputTensors.Num() == 1);
-		check(OutputTensors.Num() == 1);
-
-		int32 BlockSize;
+		check(Inputs.Num() == 1);
+		check(Outputs.Num() == 1);
 
 		const FNNEAttributeValue* BlockSizeAttr = Attributes.GetAttributeValue(TEXT("blocksize"));
 		if (BlockSizeAttr)
@@ -57,9 +75,38 @@ public:
 			return false;
 		}
 
+		Order = SpaceOrderFromModeString(Attributes.GetValueOrDefault<FString>(TEXT("mode"), FString(TEXT("DCR"))));
+
+		return true;
+	}
+
+	//
+	//
+	//
+	virtual int PrepareOutputs(TConstArrayView<NNE::Internal::FTensorRef> InputTensors, TArrayView<NNE::Internal::FTensorRef> OutputTensors) const override
+	{
+		TConstArrayView<uint32>		InputShape = InputTensors[0]->GetShape().GetData();
+		Util::FSmallUIntArray		OutputShape;
+		
+		OutputShape.SetNum(DIM_COUNT);
+		OutputShape[N] = InputShape[N];
+		OutputShape[C] = InputShape[C] / (BlockSize * BlockSize);
+		OutputShape[H] = InputShape[H] * BlockSize;
+		OutputShape[W] = InputShape[W] * BlockSize;
+
+		OutputTensors[0]->SetShape(NNE::FTensorShape::Make(OutputShape));
+
+		return 0;
+	}
+	
+	//
+	//
+	//
+	virtual bool Create(IDMLDevice* Device, TConstArrayView<NNE::Internal::FTensorRef> InputTensors, TConstArrayView<NNE::Internal::FTensorRef> OutputTensors) override
+	{
 
-		const NNE::Internal::FTensor& InputTensorDesc = InputTensors[0];
-		const NNE::Internal::FTensor& OutputTensorDesc = OutputTensors[0];
+		const NNE::Internal::FTensor& InputTensorDesc = *InputTensors[0];
+		const NNE::Internal::FTensor& OutputTensorDesc = *OutputTensors[0];
 
 		FTensorDescDml	DmlInputTensorDesc;
 		FTensorDescDml	DmlOutputTensorDesc;
@@ -87,7 +134,7 @@ public:
 		DmlDepthToSpaceOpDesc.InputTensor = DmlInputTensorDesc.GetDmlDesc();
 		DmlDepthToSpaceOpDesc.OutputTensor = DmlOutputTensorDesc.GetDmlDesc();
 		DmlDepthToSpaceOpDesc.BlockSize = BlockSize;
-		DmlDepthToSpaceOpDesc.Order = SpaceOrderFromModeString(Attributes.GetValueOrDefault<FString>(TEXT("mode"), FString(TEXT("DCR"))));
+		DmlDepthToSpaceOpDesc.Order = Order;
 
 		return CreateOperator(Device, DML_OPERATOR_DESC{ DML_OPERATOR_DEPTH_TO_SPACE1, &DmlDepthToSpaceOpDesc });
 	}
diff --git a/Engine/Plugins/Experimental/NNERuntimeRDG/Source/NNERuntimeRDG/Private/Dml/NNEDmlOperatorElementwise.cpp b/Engine/Plugins/Experimental/NNERuntimeRDG/Source/NNERuntimeRDG/Private/Dml/NNEDmlOperatorElementwise.cpp
--- a/Engine/Plugins/Experimental/NNERuntimeRDG/Source/NNERuntimeRDG/Private/Dml/NNEDmlOperatorElementwise.cpp
+++ b/Engine/Plugins/Experimental/NNERuntimeRDG/Source/NNERuntimeRDG/Private/Dml/NNEDmlOperatorElementwise.cpp
@@ -6,7 +6,6 @@
 namespace UE::NNERuntimeRDG::Private::Dml
 {
 
-
 /**
  * Element-wise unary ML operator implementation
  */
@@ -19,20 +18,26 @@ class FOperatorDmlElementWiseUnary : public FOperatorDml
 {
 public:
 
+	//
+	//
+	//
 	static FOperatorDml* Create()
 	{
 		return new FOperatorDmlElementWiseUnary();
 	}
 
+	//
+	//
+	//
 	static bool Validate(const NNE::FAttributeMap& AttributeMap, TConstArrayView<ENNETensorDataType> InputTypes, TConstArrayView<NNE::FSymbolicTensorShape> InputShapes)
 	{
-		if(InputShapes.Num() != 1)
+		if (InputShapes.Num() != 1)
 		{
 			UE_LOG(LogNNE, Warning, TEXT("Invalid number of input tensors"));
 			return false;
 		}
 
-		if(!CheckElementwiseTensor(InputTypes[0], InputShapes[0]))
+		if (!CheckElementwiseTensor(InputTypes[0], InputShapes[0]))
 		{
 			return false;
 		}
@@ -40,42 +45,50 @@ public:
 		return true;
 	}
 
+	//
+	//
+	//
 	virtual ~FOperatorDmlElementWiseUnary() = default;
 
 private:
 
-	FOperatorDmlElementWiseUnary() : Min(TNumericLimits<float>::Min()), Max(TNumericLimits<float>::Max()) {}
 	float Min;
 	float Max;
 
+	//
+	//
+	//
+	FOperatorDmlElementWiseUnary() 
+		: Min(TNumericLimits<float>::Min())
+		, Max(TNumericLimits<float>::Max()) 
+	{
+	}
+	
 public:
 
 	//
 	//
 	//
-	virtual bool Initialize(IDMLDevice* Device, TArrayView<const NNE::Internal::FTensor> InputTensors, TArrayView<const NNE::Internal::FTensor> OutputTensors, const NNE::FAttributeMap& Attributes) override
+	virtual bool Initialize(TConstArrayView<NNE::FTensorDesc> Inputs, TConstArrayView<NNE::FTensorDesc> Outputs, const NNE::FAttributeMap& Attributes) override
 	{
-
-		const NNE::Internal::FTensor& InputTensor = InputTensors[0];
-		const NNE::Internal::FTensor& OutputTensor = OutputTensors[0];
-
 		if constexpr (std::is_same_v<DmlElementWiseOpDescType, DML_ELEMENT_WISE_CLIP_OPERATOR_DESC>)
 		{
 			const FNNEAttributeValue* MinAttr = Attributes.GetAttributeValue(TEXT("min"));
-			if(MinAttr)
+			if (MinAttr)
 			{
-				if(MinAttr->GetType() != ENNEAttributeDataType::Float)
+				if (MinAttr->GetType() != ENNEAttributeDataType::Float)
 				{
 					UE_LOG(LogNNE, Error, TEXT("Min attribute of clip must be float for DML inference"));
 					return false;
 				}
-				
+
 				Min = MinAttr->GetValue<float>();
 			}
+			
 			const FNNEAttributeValue* MaxAttr = Attributes.GetAttributeValue(TEXT("max"));
-			if(MaxAttr)
+			if (MaxAttr)
 			{
-				if(MaxAttr->GetType() != ENNEAttributeDataType::Float)
+				if (MaxAttr->GetType() != ENNEAttributeDataType::Float)
 				{
 					UE_LOG(LogNNE, Error, TEXT("Max attribute of clip must be float for DML inference"));
 					return false;
@@ -84,6 +97,29 @@ public:
 			}
 		}
 
+		return true;
+	}
+
+	//
+	//
+	//
+	virtual int PrepareOutputs(TConstArrayView<NNE::Internal::FTensorRef> InputTensors, TArrayView<NNE::Internal::FTensorRef> OutputTensors) const override
+	{
+		check(InputTensors.Num() == 1);
+		check(OutputTensors.Num() == 1);
+		OutputTensors[0]->SetShape(InputTensors[0]->GetShape());
+
+		return 0;
+	}
+
+	//
+	//
+	//
+	virtual bool Create(IDMLDevice* Device, TConstArrayView<NNE::Internal::FTensorRef> InputTensors, TConstArrayView<NNE::Internal::FTensorRef> OutputTensors) override
+	{
+		const NNE::Internal::FTensor& InputTensor = *InputTensors[0];
+		const NNE::Internal::FTensor& OutputTensor = *OutputTensors[0];
+
 		// Initialize tensor descriptor (it's same for both input and output)
 		FTensorDescDml	DmlTensorDesc;
 
@@ -109,6 +145,9 @@ public:
 
 private:
 
+	//
+	//
+	//
 	template<typename OpDesc>
 	void InitDmlOpDesc(OpDesc& Desc, const DML_TENSOR_DESC& TensorDesc)
 	{
@@ -116,6 +155,9 @@ private:
 		Desc.OutputTensor = &TensorDesc;
 	}
 
+	//
+	//
+	//
 	void InitDmlOpDesc(DML_ELEMENT_WISE_CLIP_OPERATOR_DESC& Desc, const DML_TENSOR_DESC& TensorDesc)
 	{
 		Desc.InputTensor = &TensorDesc;
@@ -139,25 +181,31 @@ class FOperatorDmlElementWiseBinary : public FOperatorDml
 
 public:
 
+	//
+	//
+	//
 	static FOperatorDml* Create()
 	{
 		return new FOperatorDmlElementWiseBinary();
 	}
 
+	//
+	//
+	//
 	static bool Validate(const NNE::FAttributeMap& AttributeMap, TConstArrayView<ENNETensorDataType> InputTypes, TConstArrayView<NNE::FSymbolicTensorShape> InputShapes)
 	{
-		if(InputShapes.Num() != 2)
+		if (InputShapes.Num() != 2)
 		{
 			UE_LOG(LogNNE, Warning, TEXT("Invalid number of input tensors"));
 			return false;
 		}
 
-		if(!CheckElementwiseTensor(InputTypes[0], InputShapes[0]))
+		if (!CheckElementwiseTensor(InputTypes[0], InputShapes[0]))
 		{
 			return false;
 		}
 
-		if(!CheckElementwiseTensor(InputTypes[1], InputShapes[1]))
+		if (!CheckElementwiseTensor(InputTypes[1], InputShapes[1]))
 		{
 			return false;
 		}
@@ -167,6 +215,9 @@ public:
 
 private:
 
+	//
+	//
+	//
 	FOperatorDmlElementWiseBinary() = default;
 
 public:
@@ -174,11 +225,55 @@ public:
 	//
 	//
 	//
-	virtual bool Initialize(IDMLDevice* Device, TArrayView<const NNE::Internal::FTensor> InputTensors, TArrayView<const NNE::Internal::FTensor> OutputTensors, const NNE::FAttributeMap& Attributes) override
+	virtual bool Initialize(TConstArrayView<NNE::FTensorDesc> Inputs, TConstArrayView<NNE::FTensorDesc> Outputs, const NNE::FAttributeMap& Attributes) override
 	{
-		const NNE::Internal::FTensor& InputATensor = InputTensors[0];
-		const NNE::Internal::FTensor& InputBTensor = InputTensors[1];
-		const NNE::Internal::FTensor& OutputTensor = OutputTensors[0];
+		return true;
+	}
+
+	//
+	//
+	//
+	virtual int PrepareOutputs(TConstArrayView<NNE::Internal::FTensorRef> InputTensors, TArrayView<NNE::Internal::FTensorRef> OutputTensors) const override
+	{
+		check(InputTensors.Num() == 2);
+		check(OutputTensors.Num() == 1);
+
+		TConstArrayView<uint32> ShapeA = InputTensors[0]->GetShape().GetData();
+		TConstArrayView<uint32> ShapeB = InputTensors[1]->GetShape().GetData();
+		const int32 OutRank = FMath::Max(ShapeA.Num(), ShapeB.Num());
+		
+		TArray<uint32> OutShape;
+
+		OutShape.SetNumUninitialized(OutRank);
+		for (int32 i = 0; i < OutRank; ++i)
+		{
+			int32 IndexA = ShapeA.Num() - 1 - i;
+			int32 IndexB = ShapeB.Num() - 1 - i;
+			int32 ValueA = IndexA >= 0 ? ShapeA[IndexA] : 1;
+			int32 ValueB = IndexB >= 0 ? ShapeB[IndexB] : 1;
+			
+			if (ValueA != ValueB && ValueA != 1 && ValueB != 1)
+			{
+				UE_LOG(LogNNE, Warning, TEXT("Error while computing shape for element wise binary op, input shapes are not compatible"));
+				return -1;
+			}
+			
+			OutShape[OutRank - 1 - i] = FMath::Max(ValueA, ValueB);
+		}
+
+		OutputTensors[0]->SetShape(NNE::FTensorShape::Make(OutShape));
+
+		return 0;
+	}
+
+	//
+	//
+	//
+	virtual bool Create(IDMLDevice* Device, TConstArrayView<NNE::Internal::FTensorRef> InputTensors, TConstArrayView<NNE::Internal::FTensorRef> OutputTensors) override
+	{
+		const NNE::Internal::FTensor& InputATensor = *InputTensors[0];
+		const NNE::Internal::FTensor& InputBTensor = *InputTensors[1];
+		const NNE::Internal::FTensor& OutputTensor = *OutputTensors[0];
 
 		// Initialize tensor descriptors
 		FTensorDescDml	DmlInputATensorDesc;
@@ -223,6 +318,9 @@ public:
 
 private:
 
+	//
+	//
+	//
 	template<typename OpDesc>
 	void InitDmlOpDesc(OpDesc& Desc, const DML_TENSOR_DESC& LHSTensorDesc, const DML_TENSOR_DESC& RHSTensorDesc, const DML_TENSOR_DESC& OutputTensorDesc)
 	{
@@ -231,6 +329,9 @@ private:
 		Desc.OutputTensor = &OutputTensorDesc;
 	}
 
+	//
+	//
+	//
 	void InitDmlOpDesc(DML_ELEMENT_WISE_POW_OPERATOR_DESC& Desc, const DML_TENSOR_DESC& LHSTensorDesc, const DML_TENSOR_DESC& RHSTensorDesc, const DML_TENSOR_DESC& OutputTensorDesc)
 	{
 		Desc.InputTensor = &LHSTensorDesc;
diff --git a/Engine/Plugins/Experimental/NNERuntimeRDG/Source/NNERuntimeRDG/Private/Dml/NNEDmlOperatorFlatten.cpp b/Engine/Plugins/Experimental/NNERuntimeRDG/Source/NNERuntimeRDG/Private/Dml/NNEDmlOperatorFlatten.cpp
--- a/Engine/Plugins/Experimental/NNERuntimeRDG/Source/NNERuntimeRDG/Private/Dml/NNEDmlOperatorFlatten.cpp
+++ b/Engine/Plugins/Experimental/NNERuntimeRDG/Source/NNERuntimeRDG/Private/Dml/NNEDmlOperatorFlatten.cpp
@@ -10,6 +10,7 @@ namespace UE::NNERuntimeRDG::Private::Dml
 
 class FOperatorDmlFlatten : public FOperatorDml
 {
+	int32 Axis;
 
 public:
 
@@ -27,22 +28,36 @@ public:
 	//
 	//
 	//
-	virtual bool Initialize(IDMLDevice* Device, TArrayView<const NNE::Internal::FTensor> InputTensors, TArrayView<const NNE::Internal::FTensor> OutputTensors, const NNE::FAttributeMap& Attributes) override
+	virtual bool Initialize(TConstArrayView<NNE::FTensorDesc> Inputs, TConstArrayView<NNE::FTensorDesc> Outputs, const NNE::FAttributeMap& Attributes) override
 	{
-		check(InputTensors.Num() == 1);
-		check(OutputTensors.Num() == 1);
+		check(Inputs.Num() == 1);
+		check(Outputs.Num() == 1);
 
-		TConstArrayView<uint32>	InputShape = InputTensors[0].GetShape().GetData();
-		int32					InputRank = InputShape.Num();
-		int32					Axis = Attributes.GetValueOrDefault(TEXT("axis"), 1);
-		
+		const int32 InputRank = Inputs[0].GetShape().Rank();
+
+		Axis = Attributes.GetValueOrDefault(TEXT("axis"), 1);
 		if (Axis > InputRank || Axis < -InputRank)
 		{
-			UE_LOG(LogNNE, Warning, TEXT("Flatten 'Axis' attribute should be in the range [-r,r] with r being the rank of the input (name: %s) however axis is %d while rank is %d."), *InputTensors[0].GetName(), Axis, InputRank);
+			UE_LOG(
+				LogNNE, 
+				Warning, 
+				TEXT("Flatten 'Axis' attribute should be in the range [-r,r] with r being the rank of the input (name: %s) however axis is %d while rank is %d."), 
+				*Inputs[0].GetName(), Axis, InputRank);
+			
 			return false;
 		}
 
-		HandleNegativeAxis(Axis, InputRank);
+		Axis = HandleNegativeAxis(Axis, InputRank);
+
+		return true;
+	}
+
+	//
+	//
+	//
+	virtual int PrepareOutputs(TConstArrayView<NNE::Internal::FTensorRef> InputTensors, TArrayView<NNE::Internal::FTensorRef> OutputTensors) const override
+	{
+		TConstArrayView<uint32> InputShape = InputTensors[0]->GetShape().GetData();
 
 		uint32 InnerDimSize = 1;
 
@@ -55,13 +70,22 @@ public:
 
 		OutputShape.Reserve(InputShape.Num());
 		OutputShape.Add(InnerDimSize);
-		OutputShape.Add(InputTensors[0].GetShape().Volume() / InnerDimSize);
-		
+		OutputShape.Add(InputTensors[0]->GetShape().Volume() / InnerDimSize);
+
+		OutputTensors[0]->SetShape(NNE::FTensorShape::Make(OutputShape));
+
+		return 0;
+	}
+
+	//
+	//
+	//
+	virtual bool Create(IDMLDevice * Device, TConstArrayView<NNE::Internal::FTensorRef> InputTensors, TConstArrayView<NNE::Internal::FTensorRef> OutputTensors) override
+	{
 		FTensorDescDml DmlTensorDesc;
 		
 		if (!DmlTensorDesc
-				.SetFromTensor(OutputTensors[0])
-				.SetShape(OutputShape)
+				.SetFromTensor(*OutputTensors[0])
 				.Validate())
 		{
 			UE_LOG(LogNNE, Error, TEXT("Failed to initialize Flatten's output tensor for DML inference"));
diff --git a/Engine/Plugins/Experimental/NNERuntimeRDG/Source/NNERuntimeRDG/Private/Dml/NNEDmlOperatorGRU.cpp b/Engine/Plugins/Experimental/NNERuntimeRDG/Source/NNERuntimeRDG/Private/Dml/NNEDmlOperatorGRU.cpp
--- a/Engine/Plugins/Experimental/NNERuntimeRDG/Source/NNERuntimeRDG/Private/Dml/NNEDmlOperatorGRU.cpp
+++ b/Engine/Plugins/Experimental/NNERuntimeRDG/Source/NNERuntimeRDG/Private/Dml/NNEDmlOperatorGRU.cpp
@@ -1,313 +1,316 @@
 // Copyright Epic Games, Inc. All Rights Reserved.
 
-#ifdef NNE_USE_DIRECTML
-
-#include "NNEDmlOperator.h"
-
-namespace UE::NNERuntimeRDG::Private::Dml
-{
-
-union FDMLGRUActivationOpDescUnion
-{
-	DML_ACTIVATION_RELU_OPERATOR_DESC Relu;
-	DML_ACTIVATION_LEAKY_RELU_OPERATOR_DESC LeakyRelu;
-	DML_ACTIVATION_THRESHOLDED_RELU_OPERATOR_DESC ThresholdedRelu;
-	DML_ACTIVATION_TANH_OPERATOR_DESC Tanh;
-	DML_ACTIVATION_SCALED_TANH_OPERATOR_DESC ScaledTanh;
-	DML_ACTIVATION_SIGMOID_OPERATOR_DESC Sigmoid;
-	DML_ACTIVATION_HARD_SIGMOID_OPERATOR_DESC HardSigmoid;
-    DML_ACTIVATION_ELU_OPERATOR_DESC Elu;
-    DML_ACTIVATION_SOFTSIGN_OPERATOR_DESC Softsign;
-    DML_ACTIVATION_SOFTPLUS_OPERATOR_DESC Softplus;
-};
-
-/**
- * GRU
- */
-class FOperatorDmlGRU : public FOperatorDml
-{
-
-	static DML_RECURRENT_NETWORK_DIRECTION DirectionFromString(FStringView StringVal)
-	{
-		if (FCString::Stricmp(StringVal.GetData(), TEXT("FORWARD")) == 0)
-		{
-			return DML_RECURRENT_NETWORK_DIRECTION_FORWARD;
-		}
-		else if (FCString::Stricmp(StringVal.GetData(), TEXT("BACKWARD")) == 0)
-		{
-			return DML_RECURRENT_NETWORK_DIRECTION_BACKWARD;
-		}
-		else if (FCString::Stricmp(StringVal.GetData(), TEXT("BIDIRECTIONAL")) == 0)
-		{
-			return DML_RECURRENT_NETWORK_DIRECTION_BIDIRECTIONAL;
-		}
-		else
-		{
-			return DML_RECURRENT_NETWORK_DIRECTION_FORWARD;
-		}
-	}
-
-	bool InitActivations(const TArray<FString>& Activations, const TArray<float>& Alphas, 
-						 const TArray<float>& Betas, TArray<DML_OPERATOR_DESC>& ActivationOpDescs,
-						 TArray<FDMLGRUActivationOpDescUnion>& ActivationOpInnerDescs)
-	{
-		ActivationOpDescs.SetNumZeroed(Activations.Num());
-		ActivationOpInnerDescs.SetNumZeroed(Activations.Num());
-
-		for(int Idx = 0; Idx < Activations.Num(); Idx++)
-		{
-			const FString& Activation = Activations[Idx];
-            DML_OPERATOR_DESC& ActivationOpDesc = ActivationOpDescs[Idx];
-            FDMLGRUActivationOpDescUnion& ActivationOpInnerDesc = ActivationOpInnerDescs[Idx];
-			ActivationOpDesc.Desc = &ActivationOpInnerDesc;
-
-            if (Activation == TEXT("Relu"))
-            { 
-                ActivationOpDesc.Type = DML_OPERATOR_ACTIVATION_RELU;
-            }
-            else if (Activation == TEXT("LeakyRelu"))
-            { 
-                ActivationOpDesc.Type = DML_OPERATOR_ACTIVATION_LEAKY_RELU;
-                ActivationOpInnerDesc.LeakyRelu.Alpha = Idx >= Alphas.Num() ? .01f : Alphas[Idx];
-            }
-            else if (Activation == TEXT("ThresholdedRelu"))
-            { 
-                ActivationOpDesc.Type = DML_OPERATOR_ACTIVATION_THRESHOLDED_RELU;
-                ActivationOpInnerDesc.LeakyRelu.Alpha = Idx >= Alphas.Num() ? 1.0f : Alphas[Idx];
-            }
-            else if (Activation == TEXT("Tanh"))
-            { 
-                ActivationOpDesc.Type = DML_OPERATOR_ACTIVATION_TANH;
-            }
-            else if (Activation == TEXT("ScaledTanh"))
-            { 
-                ActivationOpDesc.Type = DML_OPERATOR_ACTIVATION_SCALED_TANH;
-                ActivationOpInnerDesc.ScaledTanh.Alpha = Idx >= Alphas.Num() ? 1.0f : Alphas[Idx];
-                ActivationOpInnerDesc.ScaledTanh.Beta = Idx >= Betas.Num() ? 1.0f : Betas[Idx];
-            }
-            else if (Activation == TEXT("Sigmoid"))
-            { 
-                ActivationOpDesc.Type = DML_OPERATOR_ACTIVATION_SIGMOID;
-            }
-            else if (Activation == TEXT("HardSigmoid"))
-            { 
-                ActivationOpDesc.Type = DML_OPERATOR_ACTIVATION_HARD_SIGMOID;
-                ActivationOpInnerDesc.HardSigmoid.Alpha = Idx >= Alphas.Num() ? 0.2f : Alphas[Idx];
-                ActivationOpInnerDesc.HardSigmoid.Beta = Idx >= Betas.Num() ? 0.5f : Betas[Idx];
-            }
-            else if (Activation == TEXT("Elu"))
-            { 
-                ActivationOpDesc.Type = DML_OPERATOR_ACTIVATION_ELU;
-                ActivationOpInnerDesc.Elu.Alpha = Idx >= Alphas.Num() ? 1.0f : Alphas[Idx];
-            }
-            else if (Activation == TEXT("Softsign"))
-            { 
-                ActivationOpDesc.Type = DML_OPERATOR_ACTIVATION_SOFTSIGN;
-            }
-            else if (Activation == TEXT("Softplus"))
-            { 
-                ActivationOpDesc.Type = DML_OPERATOR_ACTIVATION_SOFTPLUS;
-            }
-            else
-            {
-                UE_LOG(LogNNE, Error, TEXT("Failed to initialize GRU: unsupported activation: %s"), *Activation);
-				return false;
-            }
-		}
-
-		return true;
-	}
-
-public:
-
-	static FOperatorDml* Create()
-	{
-		return new FOperatorDmlGRU();
-	}
-
-	static bool Validate(const NNE::FAttributeMap& AttributeMap, TConstArrayView<ENNETensorDataType> InputTypes, TConstArrayView<NNE::FSymbolicTensorShape> InputShapes)
-	{
-		//TODO
-		return true;
-	}
-
-	//
-	//
-	//
-	virtual bool Initialize(IDMLDevice* Device, TArrayView<const NNE::Internal::FTensor> InputTensors, TArrayView<const NNE::Internal::FTensor> OutputTensors, const NNE::FAttributeMap& Attributes) override
-	{
-		
-#define SET_DMLTENSORDESC_FROM_TENSOR_COND_DATATYPE(OpDesc, Tensor, Name, DataType, Cond)\
-		FTensorDescDml	Dml##Name##Desc;\
-		if(Cond)\
-		{\
-			if (!Dml##Name##Desc\
-					.SetTensorRank(4, 4)\
-					.SetShape(Tensor.GetShape())\
-					.SetDataType(DataType)\
-					.Validate())\
-			{\
-				UE_LOG(LogNNE, Error, TEXT("Failed to initialize GRU's "#Name" for DML inference"));\
-				return false;\
-			}\
-			##OpDesc.##Name = Dml##Name##Desc.GetDmlDesc();\
-		}\
-		else\
-		{\
-			##OpDesc.##Name = nullptr;\
-		}
-
-#define SET_DMLTENSORDESC_FROM_TENSOR_COND(OpDesc, Tensor, Name, Cond)	SET_DMLTENSORDESC_FROM_TENSOR_COND_DATATYPE(OpDesc, Tensor, Name, Tensor.GetDataType(), Cond)
-
-#define SET_DMLTENSORDESC_FROM_TENSOR(OpDesc, Tensor, Name)	SET_DMLTENSORDESC_FROM_TENSOR_COND(OpDesc, Tensor, Name, true)
-
-		check(InputTensors.Num() >= 3 && InputTensors.Num() <= 6);
-		check(OutputTensors.Num() >= 0 && OutputTensors.Num() <= 2);
-
-		check(InputTensors[0].GetShape().Rank() == 3);
-		check(InputTensors[1].GetShape().Rank() == 3);
-		check(InputTensors[2].GetShape().Rank() == 3);
-		if(InputTensors.Num() >= 4)
-		{
-			check(InputTensors[3].GetShape().Rank() == 2);
-		}
-		if(InputTensors.Num() >= 5)
-		{
-			check(InputTensors[4].GetShape().Rank() == 1);
-		}
-		if(InputTensors.Num() >= 6)
-		{
-			check(InputTensors[5].GetShape().Rank() == 3);
-		}
-		if(OutputTensors.Num() >= 1)
-		{
-			check(OutputTensors[0].GetShape().Rank() == 4);
-		}
-		if(OutputTensors.Num() >= 2)
-		{
-			check(OutputTensors[1].GetShape().Rank() == 3);
-		}
-
-		const int SeqLength = InputTensors[0].GetShape().GetData()[0];
-		const int BatchSize = InputTensors[0].GetShape().GetData()[1];
-		const int InputSize = InputTensors[0].GetShape().GetData()[2];
-		const int NumDirections = InputTensors[1].GetShape().GetData()[0];
-		const int HiddenSize = InputTensors[2].GetShape().GetData()[2];
-
-		
-		check(HiddenSize * 3 == InputTensors[1].GetShape().GetData()[1]);
-		check(InputSize == InputTensors[1].GetShape().GetData()[2]);
-		check(NumDirections == InputTensors[2].GetShape().GetData()[0]);
-		check(HiddenSize * 3 == InputTensors[2].GetShape().GetData()[1]);
-		if(InputTensors.Num() >= 4)
-		{
-			check(NumDirections == InputTensors[3].GetShape().GetData()[0]);
-			check(HiddenSize * 6 == InputTensors[3].GetShape().GetData()[1]);
-		}
-		if(InputTensors.Num() >= 5)
-		{
-			check(BatchSize == InputTensors[4].GetShape().GetData()[0]);
-			check(InputTensors[4].GetDataType() == ENNETensorDataType::Int32);
-		}
-		if(InputTensors.Num() >= 6)
-		{
-			check(NumDirections == InputTensors[5].GetShape().GetData()[0]);
-			check(BatchSize == InputTensors[5].GetShape().GetData()[1]);
-			check(HiddenSize == InputTensors[5].GetShape().GetData()[2]);
-		}
-		if(OutputTensors.Num() >= 1)
-		{
-			check(SeqLength == OutputTensors[0].GetShape().GetData()[0]);
-			check(NumDirections == OutputTensors[0].GetShape().GetData()[1]);
-			check(BatchSize == OutputTensors[0].GetShape().GetData()[2]);
-			check(HiddenSize == OutputTensors[0].GetShape().GetData()[3]);
-		}
-		if(OutputTensors.Num() >= 2)
-		{
-			check(NumDirections == OutputTensors[1].GetShape().GetData()[0]);
-			check(BatchSize == OutputTensors[1].GetShape().GetData()[1]);
-			check(HiddenSize == OutputTensors[1].GetShape().GetData()[2]);
-		}
-
-		DML_GRU_OPERATOR_DESC	DmlGRUOpDesc{};
-
-		SET_DMLTENSORDESC_FROM_TENSOR(DmlGRUOpDesc, InputTensors[0], InputTensor)
-		SET_DMLTENSORDESC_FROM_TENSOR(DmlGRUOpDesc, InputTensors[1], WeightTensor)
-		SET_DMLTENSORDESC_FROM_TENSOR(DmlGRUOpDesc, InputTensors[2], RecurrenceTensor)
-		SET_DMLTENSORDESC_FROM_TENSOR_COND(DmlGRUOpDesc, InputTensors[3], BiasTensor, InputTensors.Num() >= 4)
-		// Cast SequenceLengthsTensor from int32 to uint32 due to differences in representation between ONNX and DML formats.
-		SET_DMLTENSORDESC_FROM_TENSOR_COND_DATATYPE(DmlGRUOpDesc, InputTensors[4], SequenceLengthsTensor, ENNETensorDataType::UInt32, InputTensors.Num() >= 5)
-		SET_DMLTENSORDESC_FROM_TENSOR_COND(DmlGRUOpDesc, InputTensors[5], HiddenInitTensor, InputTensors.Num() >= 6)
-		SET_DMLTENSORDESC_FROM_TENSOR_COND(DmlGRUOpDesc, OutputTensors[0], OutputSequenceTensor, OutputTensors.Num() >= 1)
-		SET_DMLTENSORDESC_FROM_TENSOR_COND(DmlGRUOpDesc, OutputTensors[1], OutputSingleTensor, OutputTensors.Num() >= 2)
-
-		DmlGRUOpDesc.Direction = DirectionFromString(Attributes.GetValueOrDefault<FString>(TEXT("direction"), TEXT("forward")));
-		
-		const TArray<FString> DefaultActivations = {"Sigmoid", "Tanh"};
-		const TArray<FString> DefaultActivationsBidirectional = {"Sigmoid", "Tanh", "Sigmoid", "Tanh"};
-		TArray<FString> Activations = 
-			Attributes.GetValueOrDefault<TArray<FString>>(
-				TEXT("activations"), 
-				DmlGRUOpDesc.Direction == DML_RECURRENT_NETWORK_DIRECTION_BIDIRECTIONAL ?
-					DefaultActivationsBidirectional : DefaultActivations);
-		
-		check(DmlGRUOpDesc.Direction == DML_RECURRENT_NETWORK_DIRECTION_BIDIRECTIONAL ?
-				(Activations.Num() == 4) : (Activations.Num() == 2));
-
-		if(Attributes.GetAttributeValue(TEXT("clip")) != nullptr)
-		{
-			UE_LOG(LogNNE, Error, TEXT("GRU's clip attribute not supported for DML inference"));
-			return false;
-		}
-		
-		const int HiddenSizeParam = Attributes.GetValueOrDefault<int32>(TEXT("hidden_size"), HiddenSize);
-		check(HiddenSize == HiddenSizeParam);
-
-		const int Layout = Attributes.GetValueOrDefault<int32>(TEXT("layout"), 0);
-		if(Layout != 0)
-		{
-			UE_LOG(LogNNE, Error, TEXT("GRU's layout attribute should be 0 for DML inference"));
-			return false;
-		}
-
-		//TODO: a way should be found to share implementation of activations from NNEDmlOperatorElementwise.cpp
-
-		TArray<float> Alphas = 
-			Attributes.GetValueOrDefault<TArray<float>>(
-				TEXT("activation_alpha"), 
-				TArray<float>());
-
-		TArray<float> Betas = 
-			Attributes.GetValueOrDefault<TArray<float>>(
-				TEXT("activation_beta"), 
-				TArray<float>());
-
-		TArray<DML_OPERATOR_DESC> ActivationOpDescs;
-		TArray<FDMLGRUActivationOpDescUnion> ActivationOpInnerDescs;
-		
-		if(!InitActivations(Activations, Alphas, 
-						 Betas, ActivationOpDescs,
-						 ActivationOpInnerDescs))
-		{
-			return false;
-		}
-
-		DmlGRUOpDesc.ActivationDescCount = (uint32_t) ActivationOpDescs.Num();
-		DmlGRUOpDesc.ActivationDescs = ActivationOpDescs.GetData();
-		
-		DmlGRUOpDesc.LinearBeforeReset = 
-			Attributes.GetValueOrDefault<int32>(TEXT("linear_before_reset"), 0) == 0 ? false : true;
-
-		return CreateOperator(Device, DML_OPERATOR_DESC{ DML_OPERATOR_GRU, &DmlGRUOpDesc });
-
-#undef SET_DMLTENSORDESC_FROM_TENSOR
-#undef SET_DMLTENSORDESC_FROM_TENSOR_COND
-	}
-};
-
-// Register GRU operator on Module startup
-NNE_DML_REGISTER_OP(GRU)
-
-} // namespace UE::NNERuntimeRDG::Private::Dml
-
-#endif // NNE_USE_DIRECTML
+// NOTE: For now we're not using GRU operator
+
+//
+//#ifdef NNE_USE_DIRECTML
+//
+//#include "NNEDmlOperator.h"
+//
+//namespace UE::NNERuntimeRDG::Private::Dml
+//{
+//
+//union FDMLGRUActivationOpDescUnion
+//{
+//	DML_ACTIVATION_RELU_OPERATOR_DESC Relu;
+//	DML_ACTIVATION_LEAKY_RELU_OPERATOR_DESC LeakyRelu;
+//	DML_ACTIVATION_THRESHOLDED_RELU_OPERATOR_DESC ThresholdedRelu;
+//	DML_ACTIVATION_TANH_OPERATOR_DESC Tanh;
+//	DML_ACTIVATION_SCALED_TANH_OPERATOR_DESC ScaledTanh;
+//	DML_ACTIVATION_SIGMOID_OPERATOR_DESC Sigmoid;
+//	DML_ACTIVATION_HARD_SIGMOID_OPERATOR_DESC HardSigmoid;
+//    DML_ACTIVATION_ELU_OPERATOR_DESC Elu;
+//    DML_ACTIVATION_SOFTSIGN_OPERATOR_DESC Softsign;
+//    DML_ACTIVATION_SOFTPLUS_OPERATOR_DESC Softplus;
+//};
+//
+///**
+// * GRU
+// */
+//class FOperatorDmlGRU : public FOperatorDml
+//{
+//
+//	static DML_RECURRENT_NETWORK_DIRECTION DirectionFromString(FStringView StringVal)
+//	{
+//		if (FCString::Stricmp(StringVal.GetData(), TEXT("FORWARD")) == 0)
+//		{
+//			return DML_RECURRENT_NETWORK_DIRECTION_FORWARD;
+//		}
+//		else if (FCString::Stricmp(StringVal.GetData(), TEXT("BACKWARD")) == 0)
+//		{
+//			return DML_RECURRENT_NETWORK_DIRECTION_BACKWARD;
+//		}
+//		else if (FCString::Stricmp(StringVal.GetData(), TEXT("BIDIRECTIONAL")) == 0)
+//		{
+//			return DML_RECURRENT_NETWORK_DIRECTION_BIDIRECTIONAL;
+//		}
+//		else
+//		{
+//			return DML_RECURRENT_NETWORK_DIRECTION_FORWARD;
+//		}
+//	}
+//
+//	bool InitActivations(const TArray<FString>& Activations, const TArray<float>& Alphas, 
+//						 const TArray<float>& Betas, TArray<DML_OPERATOR_DESC>& ActivationOpDescs,
+//						 TArray<FDMLGRUActivationOpDescUnion>& ActivationOpInnerDescs)
+//	{
+//		ActivationOpDescs.SetNumZeroed(Activations.Num());
+//		ActivationOpInnerDescs.SetNumZeroed(Activations.Num());
+//
+//		for(int Idx = 0; Idx < Activations.Num(); Idx++)
+//		{
+//			const FString& Activation = Activations[Idx];
+//            DML_OPERATOR_DESC& ActivationOpDesc = ActivationOpDescs[Idx];
+//            FDMLGRUActivationOpDescUnion& ActivationOpInnerDesc = ActivationOpInnerDescs[Idx];
+//			ActivationOpDesc.Desc = &ActivationOpInnerDesc;
+//
+//            if (Activation == TEXT("Relu"))
+//            { 
+//                ActivationOpDesc.Type = DML_OPERATOR_ACTIVATION_RELU;
+//            }
+//            else if (Activation == TEXT("LeakyRelu"))
+//            { 
+//                ActivationOpDesc.Type = DML_OPERATOR_ACTIVATION_LEAKY_RELU;
+//                ActivationOpInnerDesc.LeakyRelu.Alpha = Idx >= Alphas.Num() ? .01f : Alphas[Idx];
+//            }
+//            else if (Activation == TEXT("ThresholdedRelu"))
+//            { 
+//                ActivationOpDesc.Type = DML_OPERATOR_ACTIVATION_THRESHOLDED_RELU;
+//                ActivationOpInnerDesc.LeakyRelu.Alpha = Idx >= Alphas.Num() ? 1.0f : Alphas[Idx];
+//            }
+//            else if (Activation == TEXT("Tanh"))
+//            { 
+//                ActivationOpDesc.Type = DML_OPERATOR_ACTIVATION_TANH;
+//            }
+//            else if (Activation == TEXT("ScaledTanh"))
+//            { 
+//                ActivationOpDesc.Type = DML_OPERATOR_ACTIVATION_SCALED_TANH;
+//                ActivationOpInnerDesc.ScaledTanh.Alpha = Idx >= Alphas.Num() ? 1.0f : Alphas[Idx];
+//                ActivationOpInnerDesc.ScaledTanh.Beta = Idx >= Betas.Num() ? 1.0f : Betas[Idx];
+//            }
+//            else if (Activation == TEXT("Sigmoid"))
+//            { 
+//                ActivationOpDesc.Type = DML_OPERATOR_ACTIVATION_SIGMOID;
+//            }
+//            else if (Activation == TEXT("HardSigmoid"))
+//            { 
+//                ActivationOpDesc.Type = DML_OPERATOR_ACTIVATION_HARD_SIGMOID;
+//                ActivationOpInnerDesc.HardSigmoid.Alpha = Idx >= Alphas.Num() ? 0.2f : Alphas[Idx];
+//                ActivationOpInnerDesc.HardSigmoid.Beta = Idx >= Betas.Num() ? 0.5f : Betas[Idx];
+//            }
+//            else if (Activation == TEXT("Elu"))
+//            { 
+//                ActivationOpDesc.Type = DML_OPERATOR_ACTIVATION_ELU;
+//                ActivationOpInnerDesc.Elu.Alpha = Idx >= Alphas.Num() ? 1.0f : Alphas[Idx];
+//            }
+//            else if (Activation == TEXT("Softsign"))
+//            { 
+//                ActivationOpDesc.Type = DML_OPERATOR_ACTIVATION_SOFTSIGN;
+//            }
+//            else if (Activation == TEXT("Softplus"))
+//            { 
+//                ActivationOpDesc.Type = DML_OPERATOR_ACTIVATION_SOFTPLUS;
+//            }
+//            else
+//            {
+//                UE_LOG(LogNNE, Error, TEXT("Failed to initialize GRU: unsupported activation: %s"), *Activation);
+//				return false;
+//            }
+//		}
+//
+//		return true;
+//	}
+//
+//public:
+//
+//	static FOperatorDml* Create()
+//	{
+//		return new FOperatorDmlGRU();
+//	}
+//
+//	static bool Validate(const NNE::FAttributeMap& AttributeMap, TConstArrayView<ENNETensorDataType> InputTypes, TConstArrayView<NNE::FSymbolicTensorShape> InputShapes)
+//	{
+//		//TODO
+//		return true;
+//	}
+//
+//	//
+//	//
+//	//
+//	virtual bool Initialize(IDMLDevice* Device, TArrayView<const NNE::Internal::FTensor> InputTensors, TArrayView<const NNE::Internal::FTensor> OutputTensors, const NNE::FAttributeMap& Attributes) override
+//	{
+//		
+//#define SET_DMLTENSORDESC_FROM_TENSOR_COND_DATATYPE(OpDesc, Tensor, Name, DataType, Cond)\
+//		FTensorDescDml	Dml##Name##Desc;\
+//		if(Cond)\
+//		{\
+//			if (!Dml##Name##Desc\
+//					.SetTensorRank(4, 4)\
+//					.SetShape(Tensor.GetShape())\
+//					.SetDataType(DataType)\
+//					.Validate())\
+//			{\
+//				UE_LOG(LogNNE, Error, TEXT("Failed to initialize GRU's "#Name" for DML inference"));\
+//				return false;\
+//			}\
+//			##OpDesc.##Name = Dml##Name##Desc.GetDmlDesc();\
+//		}\
+//		else\
+//		{\
+//			##OpDesc.##Name = nullptr;\
+//		}
+//
+//#define SET_DMLTENSORDESC_FROM_TENSOR_COND(OpDesc, Tensor, Name, Cond)	SET_DMLTENSORDESC_FROM_TENSOR_COND_DATATYPE(OpDesc, Tensor, Name, Tensor.GetDataType(), Cond)
+//
+//#define SET_DMLTENSORDESC_FROM_TENSOR(OpDesc, Tensor, Name)	SET_DMLTENSORDESC_FROM_TENSOR_COND(OpDesc, Tensor, Name, true)
+//
+//		check(InputTensors.Num() >= 3 && InputTensors.Num() <= 6);
+//		check(OutputTensors.Num() >= 0 && OutputTensors.Num() <= 2);
+//
+//		check(InputTensors[0].GetShape().Rank() == 3);
+//		check(InputTensors[1].GetShape().Rank() == 3);
+//		check(InputTensors[2].GetShape().Rank() == 3);
+//		if(InputTensors.Num() >= 4)
+//		{
+//			check(InputTensors[3].GetShape().Rank() == 2);
+//		}
+//		if(InputTensors.Num() >= 5)
+//		{
+//			check(InputTensors[4].GetShape().Rank() == 1);
+//		}
+//		if(InputTensors.Num() >= 6)
+//		{
+//			check(InputTensors[5].GetShape().Rank() == 3);
+//		}
+//		if(OutputTensors.Num() >= 1)
+//		{
+//			check(OutputTensors[0].GetShape().Rank() == 4);
+//		}
+//		if(OutputTensors.Num() >= 2)
+//		{
+//			check(OutputTensors[1].GetShape().Rank() == 3);
+//		}
+//
+//		const int SeqLength = InputTensors[0].GetShape().GetData()[0];
+//		const int BatchSize = InputTensors[0].GetShape().GetData()[1];
+//		const int InputSize = InputTensors[0].GetShape().GetData()[2];
+//		const int NumDirections = InputTensors[1].GetShape().GetData()[0];
+//		const int HiddenSize = InputTensors[2].GetShape().GetData()[2];
+//
+//		
+//		check(HiddenSize * 3 == InputTensors[1].GetShape().GetData()[1]);
+//		check(InputSize == InputTensors[1].GetShape().GetData()[2]);
+//		check(NumDirections == InputTensors[2].GetShape().GetData()[0]);
+//		check(HiddenSize * 3 == InputTensors[2].GetShape().GetData()[1]);
+//		if(InputTensors.Num() >= 4)
+//		{
+//			check(NumDirections == InputTensors[3].GetShape().GetData()[0]);
+//			check(HiddenSize * 6 == InputTensors[3].GetShape().GetData()[1]);
+//		}
+//		if(InputTensors.Num() >= 5)
+//		{
+//			check(BatchSize == InputTensors[4].GetShape().GetData()[0]);
+//			check(InputTensors[4].GetDataType() == ENNETensorDataType::Int32);
+//		}
+//		if(InputTensors.Num() >= 6)
+//		{
+//			check(NumDirections == InputTensors[5].GetShape().GetData()[0]);
+//			check(BatchSize == InputTensors[5].GetShape().GetData()[1]);
+//			check(HiddenSize == InputTensors[5].GetShape().GetData()[2]);
+//		}
+//		if(OutputTensors.Num() >= 1)
+//		{
+//			check(SeqLength == OutputTensors[0].GetShape().GetData()[0]);
+//			check(NumDirections == OutputTensors[0].GetShape().GetData()[1]);
+//			check(BatchSize == OutputTensors[0].GetShape().GetData()[2]);
+//			check(HiddenSize == OutputTensors[0].GetShape().GetData()[3]);
+//		}
+//		if(OutputTensors.Num() >= 2)
+//		{
+//			check(NumDirections == OutputTensors[1].GetShape().GetData()[0]);
+//			check(BatchSize == OutputTensors[1].GetShape().GetData()[1]);
+//			check(HiddenSize == OutputTensors[1].GetShape().GetData()[2]);
+//		}
+//
+//		DML_GRU_OPERATOR_DESC	DmlGRUOpDesc{};
+//
+//		SET_DMLTENSORDESC_FROM_TENSOR(DmlGRUOpDesc, InputTensors[0], InputTensor)
+//		SET_DMLTENSORDESC_FROM_TENSOR(DmlGRUOpDesc, InputTensors[1], WeightTensor)
+//		SET_DMLTENSORDESC_FROM_TENSOR(DmlGRUOpDesc, InputTensors[2], RecurrenceTensor)
+//		SET_DMLTENSORDESC_FROM_TENSOR_COND(DmlGRUOpDesc, InputTensors[3], BiasTensor, InputTensors.Num() >= 4)
+//		// Cast SequenceLengthsTensor from int32 to uint32 due to differences in representation between ONNX and DML formats.
+//		SET_DMLTENSORDESC_FROM_TENSOR_COND_DATATYPE(DmlGRUOpDesc, InputTensors[4], SequenceLengthsTensor, ENNETensorDataType::UInt32, InputTensors.Num() >= 5)
+//		SET_DMLTENSORDESC_FROM_TENSOR_COND(DmlGRUOpDesc, InputTensors[5], HiddenInitTensor, InputTensors.Num() >= 6)
+//		SET_DMLTENSORDESC_FROM_TENSOR_COND(DmlGRUOpDesc, OutputTensors[0], OutputSequenceTensor, OutputTensors.Num() >= 1)
+//		SET_DMLTENSORDESC_FROM_TENSOR_COND(DmlGRUOpDesc, OutputTensors[1], OutputSingleTensor, OutputTensors.Num() >= 2)
+//
+//		DmlGRUOpDesc.Direction = DirectionFromString(Attributes.GetValueOrDefault<FString>(TEXT("direction"), TEXT("forward")));
+//		
+//		const TArray<FString> DefaultActivations = {"Sigmoid", "Tanh"};
+//		const TArray<FString> DefaultActivationsBidirectional = {"Sigmoid", "Tanh", "Sigmoid", "Tanh"};
+//		TArray<FString> Activations = 
+//			Attributes.GetValueOrDefault<TArray<FString>>(
+//				TEXT("activations"), 
+//				DmlGRUOpDesc.Direction == DML_RECURRENT_NETWORK_DIRECTION_BIDIRECTIONAL ?
+//					DefaultActivationsBidirectional : DefaultActivations);
+//		
+//		check(DmlGRUOpDesc.Direction == DML_RECURRENT_NETWORK_DIRECTION_BIDIRECTIONAL ?
+//				(Activations.Num() == 4) : (Activations.Num() == 2));
+//
+//		if(Attributes.GetAttributeValue(TEXT("clip")) != nullptr)
+//		{
+//			UE_LOG(LogNNE, Error, TEXT("GRU's clip attribute not supported for DML inference"));
+//			return false;
+//		}
+//		
+//		const int HiddenSizeParam = Attributes.GetValueOrDefault<int32>(TEXT("hidden_size"), HiddenSize);
+//		check(HiddenSize == HiddenSizeParam);
+//
+//		const int Layout = Attributes.GetValueOrDefault<int32>(TEXT("layout"), 0);
+//		if(Layout != 0)
+//		{
+//			UE_LOG(LogNNE, Error, TEXT("GRU's layout attribute should be 0 for DML inference"));
+//			return false;
+//		}
+//
+//		//TODO: a way should be found to share implementation of activations from NNEDmlOperatorElementwise.cpp
+//
+//		TArray<float> Alphas = 
+//			Attributes.GetValueOrDefault<TArray<float>>(
+//				TEXT("activation_alpha"), 
+//				TArray<float>());
+//
+//		TArray<float> Betas = 
+//			Attributes.GetValueOrDefault<TArray<float>>(
+//				TEXT("activation_beta"), 
+//				TArray<float>());
+//
+//		TArray<DML_OPERATOR_DESC> ActivationOpDescs;
+//		TArray<FDMLGRUActivationOpDescUnion> ActivationOpInnerDescs;
+//		
+//		if(!InitActivations(Activations, Alphas, 
+//						 Betas, ActivationOpDescs,
+//						 ActivationOpInnerDescs))
+//		{
+//			return false;
+//		}
+//
+//		DmlGRUOpDesc.ActivationDescCount = (uint32_t) ActivationOpDescs.Num();
+//		DmlGRUOpDesc.ActivationDescs = ActivationOpDescs.GetData();
+//		
+//		DmlGRUOpDesc.LinearBeforeReset = 
+//			Attributes.GetValueOrDefault<int32>(TEXT("linear_before_reset"), 0) == 0 ? false : true;
+//
+//		return CreateOperator(Device, DML_OPERATOR_DESC{ DML_OPERATOR_GRU, &DmlGRUOpDesc });
+//
+//#undef SET_DMLTENSORDESC_FROM_TENSOR
+//#undef SET_DMLTENSORDESC_FROM_TENSOR_COND
+//	}
+//};
+//
+//// Register GRU operator on Module startup
+//NNE_DML_REGISTER_OP(GRU)
+//
+//} // namespace UE::NNERuntimeRDG::Private::Dml
+//
+//#endif // NNE_USE_DIRECTML
diff --git a/Engine/Plugins/Experimental/NNERuntimeRDG/Source/NNERuntimeRDG/Private/Dml/NNEDmlOperatorGather.cpp b/Engine/Plugins/Experimental/NNERuntimeRDG/Source/NNERuntimeRDG/Private/Dml/NNEDmlOperatorGather.cpp
--- a/Engine/Plugins/Experimental/NNERuntimeRDG/Source/NNERuntimeRDG/Private/Dml/NNEDmlOperatorGather.cpp
+++ b/Engine/Plugins/Experimental/NNERuntimeRDG/Source/NNERuntimeRDG/Private/Dml/NNEDmlOperatorGather.cpp
@@ -9,6 +9,8 @@ namespace UE::NNERuntimeRDG::Private::Dml
 
 class FOperatorDmlGather : public FOperatorDml
 {
+	int32	Axis;
+
 public:
 
 	static FOperatorDml* Create()
@@ -25,19 +27,16 @@ public:
 	//
 	//
 	//
-	virtual bool Initialize(IDMLDevice* Device, TArrayView<const NNE::Internal::FTensor> InputTensors, TArrayView<const NNE::Internal::FTensor> OutputTensors, const NNE::FAttributeMap& Attributes) override
+	virtual bool Initialize(TConstArrayView<NNE::FTensorDesc> Inputs, TConstArrayView<NNE::FTensorDesc> Outputs, const NNE::FAttributeMap& Attributes) override
 	{
-		check(InputTensors.Num() == 2);
-		check(OutputTensors.Num() == 1);
-		
-		const NNE::Internal::FTensor& InputTensor = InputTensors[0];
-		const NNE::Internal::FTensor& IndicesTensor = InputTensors[1];
-		const NNE::Internal::FTensor& OutputTensor = OutputTensors[0];
+		check(Inputs.Num() == 2);
+		check(Outputs.Num() == 1);
 
-		const NNE::FTensorShape& InputShape = InputTensor.GetShape();
-		const NNE::FTensorShape& IndicesShape = IndicesTensor.GetShape();
+		const NNE::FTensorDesc& InputTensor = Inputs[0];
+		const NNE::FTensorDesc& IndicesTensor = Inputs[1];
+		const NNE::FTensorDesc& OutputTensor = Outputs[0];
 
-		if (IndicesShape.Rank() > InputShape.Rank())
+		if (IndicesTensor.GetShape().Rank() > InputTensor.GetShape().Rank())
 		{
 			UE_LOG(LogNNE, Warning, TEXT("Indices tensor rank must match input tensor rank"));
 			return false;
@@ -59,10 +58,20 @@ public:
 		}
 
 		// Read attributes
-		int32	Axis = Attributes.GetValueOrDefault<int>(TEXT("axis"), 0);
-
+		Axis = Attributes.GetValueOrDefault<int>(TEXT("axis"), 0);
 		Axis = HandleNegativeAxis(Axis, InputTensor.GetShape().Rank());
 
+		return true;
+	}
+
+	//
+	//
+	//
+	virtual int PrepareOutputs(TConstArrayView<NNE::Internal::FTensorRef> InputTensors, TArrayView<NNE::Internal::FTensorRef> OutputTensors) const override
+	{
+		const NNE::FTensorShape& InputShape = InputTensors[0]->GetShape();
+		const NNE::FTensorShape& IndicesShape = InputTensors[1]->GetShape();
+
 		// Compute output shape
 		const int32 OutputRank = IndicesShape.Rank() + InputShape.Rank() - 1;
 		Util::FSmallUIntArray OutputShape;
@@ -81,12 +90,26 @@ public:
 			OutputShape.Add(InputShape.GetData()[DataRankIdx]);
 		}
 
-		if (OutputTensor.GetShape().Rank() != OutputRank)
+		if (OutputShape.Num() != OutputRank)
 		{
 			UE_LOG(LogNNE, Warning, TEXT("Output tensor rank must match computed output tensor rank"));
 			return false;
 		}
 
+		OutputTensors[0]->SetShape(NNE::FTensorShape::Make(OutputShape));
+
+		return 0;
+	}
+
+	//
+	//
+	//
+	virtual bool Create(IDMLDevice* Device, TConstArrayView<NNE::Internal::FTensorRef> InputTensors, TConstArrayView<NNE::Internal::FTensorRef> OutputTensors) override
+	{
+		const NNE::Internal::FTensor& InputTensor	= *InputTensors[0];
+		const NNE::Internal::FTensor& IndicesTensor = *InputTensors[1];
+		const NNE::Internal::FTensor& OutputTensor	= *OutputTensors[0];
+
 		// Initialize tensor descriptors
 		FTensorDescDml	DmlInputTensorDesc;
 		FTensorDescDml	DmlIndicesTensorDesc;
@@ -110,7 +133,7 @@ public:
 
 		if (!DmlOutputTensorDesc
 				.SetFromTensor(OutputTensor)
-				.SetShape(OutputShape, InputTensor.GetShape().Rank())
+				.SetShape(OutputTensor.GetShape(), InputTensor.GetShape().Rank())
 				.Validate())
 		{
 			UE_LOG(LogNNE, Error, TEXT("Failed to initialize tensor(s) for DML inference"));
diff --git a/Engine/Plugins/Experimental/NNERuntimeRDG/Source/NNERuntimeRDG/Private/Dml/NNEDmlOperatorGemm.cpp b/Engine/Plugins/Experimental/NNERuntimeRDG/Source/NNERuntimeRDG/Private/Dml/NNEDmlOperatorGemm.cpp
--- a/Engine/Plugins/Experimental/NNERuntimeRDG/Source/NNERuntimeRDG/Private/Dml/NNEDmlOperatorGemm.cpp
+++ b/Engine/Plugins/Experimental/NNERuntimeRDG/Source/NNERuntimeRDG/Private/Dml/NNEDmlOperatorGemm.cpp
@@ -11,6 +11,11 @@ namespace UE::NNERuntimeRDG::Private::Dml
  */
 class FOperatorDmlGemm : public FOperatorDml
 {
+	float Alpha = 1.0f;
+	float Beta = 1.0f;
+	int32 TransA = 0;
+	int32 TransB = 0;
+
 public:
 
 	static FOperatorDml* Create()
@@ -27,22 +32,76 @@ public:
 	//
 	//
 	//
-	virtual bool Initialize(IDMLDevice* Device, TArrayView<const NNE::Internal::FTensor> InputTensors, TArrayView<const NNE::Internal::FTensor> OutputTensors, const NNE::FAttributeMap& Attributes) override
+	virtual bool Initialize(TConstArrayView<NNE::FTensorDesc> Inputs, TConstArrayView<NNE::FTensorDesc> Outputs, const NNE::FAttributeMap& Attributes) override
 	{
 		// Setup attributes
-		float Alpha = 1.0f;
-		float Beta = 1.0f;
-		int32 TransA = 0;
-		int32 TransB = 0;
-		
 		Alpha = Attributes.GetValueOrDefault(TEXT("alpha"), Alpha);
 		Beta = Attributes.GetValueOrDefault(TEXT("beta"), Beta);
 		TransA = Attributes.GetValueOrDefault(TEXT("transA"), TransA);
 		TransB = Attributes.GetValueOrDefault(TEXT("transB"), TransB);
 		
-		const NNE::Internal::FTensor& InputATensor = InputTensors[0];
-		const NNE::Internal::FTensor& InputBTensor = InputTensors[1];
-		const NNE::Internal::FTensor& OutputTensor = OutputTensors[0];
+		const NNE::FTensorDesc& InputA = Inputs[0];
+		const NNE::FTensorDesc& InputB = Inputs[1];
+
+		if (InputA.GetShape().Rank() < 2 || InputA.GetShape().Rank() > 4)
+		{
+			UE_LOG(LogNNE, Error, TEXT("Gemm InputA tensor rank needs to be [2,4]"));
+			return false;
+		}
+
+		if (InputB.GetShape().Rank() < 2 || InputB.GetShape().Rank() > 4)
+		{
+			UE_LOG(LogNNE, Error, TEXT("Gemm InputB tensor rank needs to be [2,4]"));
+			return false;
+		}
+
+		if (Inputs.Num() == 3)
+		{
+			const NNE::FTensorDesc& InputC = Inputs[2];
+
+			if (InputC.GetShape().Rank() > 4)
+			{
+				UE_LOG(LogNNE, Error, TEXT("Gemm InputC tensor rank needs to be max rank 4"));
+				return false;
+			}
+		}
+
+		return true;
+	}
+
+	//
+	//
+	//
+	virtual int PrepareOutputs(TConstArrayView<NNE::Internal::FTensorRef> InputTensors, TArrayView<NNE::Internal::FTensorRef> OutputTensors) const override
+	{
+		check(InputTensors.Num() >= 2 && InputTensors.Num() <= 3);
+		check(OutputTensors.Num() == 1);
+
+		const NNE::FTensorShape& InputA = InputTensors[0]->GetShape();
+		const NNE::FTensorShape& InputB = InputTensors[1]->GetShape();
+		
+		checkf(InputA.Rank() >= 2, TEXT("Gemm InputA needs to have tensor rank at least size of 2"));
+		checkf(InputB.Rank() >= 2, TEXT("Gemm InputB needs to have tensor rank at least size of 2"));
+
+		const uint32 M = TransA != 0 ? InputA.GetData()[1] : InputA.GetData()[0];
+		const uint32 N = TransB != 0 ? InputB.GetData()[0] : InputB.GetData()[1];
+		
+		TArray<uint32> OutputShape;
+		OutputShape.Emplace(M);
+		OutputShape.Emplace(N);
+
+		OutputTensors[0]->SetShape(NNE::FTensorShape::Make(OutputShape));
+		return 0;
+	}
+
+	//
+	//
+	//
+	virtual bool Create(IDMLDevice* Device, TConstArrayView<NNE::Internal::FTensorRef> InputTensors, TConstArrayView<NNE::Internal::FTensorRef> OutputTensors) override
+	{
+		const NNE::Internal::FTensor& InputATensor = *InputTensors[0];
+		const NNE::Internal::FTensor& InputBTensor = *InputTensors[1];
+		const NNE::Internal::FTensor& OutputTensor = *OutputTensors[0];
 
 		// Initialize tensor descriptors
 		FTensorDescDml	DmlInputATensorDesc;
@@ -70,7 +129,7 @@ public:
 
 		if (InputTensors.Num() > 2)
 		{
-			const NNE::Internal::FTensor& InputCTensor = InputTensors[2];
+			const NNE::Internal::FTensor& InputCTensor = *InputTensors[2];
 
 			if (!DmlInputCTensorDesc
 					.SetTensorRank(2, 4)
diff --git a/Engine/Plugins/Experimental/NNERuntimeRDG/Source/NNERuntimeRDG/Private/Dml/NNEDmlOperatorIdentity.cpp b/Engine/Plugins/Experimental/NNERuntimeRDG/Source/NNERuntimeRDG/Private/Dml/NNEDmlOperatorIdentity.cpp
--- a/Engine/Plugins/Experimental/NNERuntimeRDG/Source/NNERuntimeRDG/Private/Dml/NNEDmlOperatorIdentity.cpp
+++ b/Engine/Plugins/Experimental/NNERuntimeRDG/Source/NNERuntimeRDG/Private/Dml/NNEDmlOperatorIdentity.cpp
@@ -8,16 +8,25 @@
 namespace UE::NNERuntimeRDG::Private::Dml
 {
 
+//
+//
+//
 class FOperatorDmlIdentity : public FOperatorDml
 {
 
 public:
 
+	//
+	//
+	//
 	static FOperatorDml* Create()
 	{
 		return new FOperatorDmlIdentity();
 	}
 
+	//
+	//
+	//
 	static bool Validate(const NNE::FAttributeMap& AttributeMap, TConstArrayView<ENNETensorDataType> InputTypes, TConstArrayView<NNE::FSymbolicTensorShape> InputShapes)
 	{
 		//TODO
@@ -27,13 +36,30 @@ public:
 	//
 	//
 	//
-	virtual bool Initialize(IDMLDevice* Device, TArrayView<const NNE::Internal::FTensor> InputTensors, TArrayView<const NNE::Internal::FTensor> OutputTensors, const NNE::FAttributeMap& Attributes) override
+	virtual bool Initialize(TConstArrayView<NNE::FTensorDesc> Inputs, TConstArrayView<NNE::FTensorDesc> Outputs, const NNE::FAttributeMap& Attributes) override
 	{
-		check(InputTensors.Num() == 1);
-		check(OutputTensors.Num() == 1);
+		check(Inputs.Num() == 1);
+		check(Outputs.Num() == 1);
+		
+		return true;
+	}
 
-		const NNE::Internal::FTensor& InputTensor = InputTensors[0];
-		const NNE::Internal::FTensor& OutputTensor = OutputTensors[0];
+	//
+	//
+	//
+	virtual int PrepareOutputs(TConstArrayView<NNE::Internal::FTensorRef> InputTensors, TArrayView<NNE::Internal::FTensorRef> OutputTensors) const override
+	{
+		OutputTensors[0]->SetShape(InputTensors[0]->GetShape());
+		return 0;
+	}
+
+	//
+	//
+	//
+	virtual bool Create(IDMLDevice* Device, TConstArrayView<NNE::Internal::FTensorRef> InputTensors, TConstArrayView<NNE::Internal::FTensorRef> OutputTensors) override
+	{
+		const NNE::Internal::FTensor& InputTensor = *InputTensors[0];
+		const NNE::Internal::FTensor& OutputTensor = *OutputTensors[0];
 
 		FTensorDescDml	DmlInputTensorDesc;
 		FTensorDescDml	DmlOutputTensorDesc;
diff --git a/Engine/Plugins/Experimental/NNERuntimeRDG/Source/NNERuntimeRDG/Private/Dml/NNEDmlOperatorInstanceNormalization.cpp b/Engine/Plugins/Experimental/NNERuntimeRDG/Source/NNERuntimeRDG/Private/Dml/NNEDmlOperatorInstanceNormalization.cpp
--- a/Engine/Plugins/Experimental/NNERuntimeRDG/Source/NNERuntimeRDG/Private/Dml/NNEDmlOperatorInstanceNormalization.cpp
+++ b/Engine/Plugins/Experimental/NNERuntimeRDG/Source/NNERuntimeRDG/Private/Dml/NNEDmlOperatorInstanceNormalization.cpp
@@ -6,12 +6,20 @@
 namespace UE::NNERuntimeRDG::Private::Dml
 {
 
+//
+//
+//
 class FOperatorDmlInstanceNormalization : public FOperatorDml
 {
 	static constexpr float DefaultEpsilon = 0.00001f;
 
+	float Epsilon;
+
 public:
 
+	//
+	//
+	//
 	static FOperatorDml* Create()
 	{
 		return new FOperatorDmlInstanceNormalization();
@@ -26,22 +34,33 @@ public:
 	//
 	//
 	//
-	virtual bool Initialize(IDMLDevice* Device, TArrayView<const NNE::Internal::FTensor> InputTensors, TArrayView<const NNE::Internal::FTensor> OutputTensors, const NNE::FAttributeMap& Attributes) override
+	virtual bool Initialize(TConstArrayView<NNE::FTensorDesc> Inputs, TConstArrayView<NNE::FTensorDesc> Outputs, const NNE::FAttributeMap& Attributes) override
 	{
-		check(InputTensors.Num() >= 1 && InputTensors.Num() <= 3);
-		check(OutputTensors.Num() == 1);
+		check(Inputs.Num() >= 1 && Inputs.Num() <= 3);
+		check(Outputs.Num() == 1);
 
-		const NNE::Internal::FTensor& InputTensor = InputTensors[0];
-		const NNE::Internal::FTensor& OutputTensor = OutputTensors[0];
+		// Read attributes
+		Epsilon = Attributes.GetValueOrDefault<float>(TEXT("epsilon"), DefaultEpsilon);
 
-		if (InputTensor.GetShape().Rank() > 8)
-		{
-			UE_LOG(LogNNE, Warning, TEXT("InstanceNormalization:InputTensor rank should be between 1 and 8, got:%d"), InputTensor.GetShape().Rank());
-			return false;
-		}
+		return true;
+	}
 
-		// Read attributes
-		float	Epsilon = Attributes.GetValueOrDefault<float>(TEXT("epsilon"), DefaultEpsilon);
+	//
+	//
+	//
+	virtual int PrepareOutputs(TConstArrayView<NNE::Internal::FTensorRef> InputTensors, TArrayView<NNE::Internal::FTensorRef> OutputTensors) const override
+	{
+		OutputTensors[0]->SetShape(InputTensors[0]->GetShape());
+		return 0;
+	}
+
+	//
+	//
+	//
+	virtual bool Create(IDMLDevice* Device, TConstArrayView<NNE::Internal::FTensorRef> InputTensors, TConstArrayView<NNE::Internal::FTensorRef> OutputTensors) override
+	{
+		const NNE::Internal::FTensor& InputTensor = *InputTensors[0];
+		const NNE::Internal::FTensor& OutputTensor = *OutputTensors[0];
 
 		// Initialize tensor descriptors
 		FTensorDescDml	DmlInputTensorDesc;
@@ -61,7 +80,7 @@ public:
 
 		if (InputTensors.Num() > 1)
 		{
-			const NNE::Internal::FTensor& ScaleTensor = InputTensors[1];
+			const NNE::Internal::FTensor& ScaleTensor = *InputTensors[1];
 
 			if (!DmlScalingTensorDesc
 					.SetTensorRank(4, 4)
@@ -75,7 +94,7 @@ public:
 
 		if (InputTensors.Num() > 2)
 		{
-			const NNE::Internal::FTensor& BiasTensor = InputTensors[2];
+			const NNE::Internal::FTensor& BiasTensor = *InputTensors[2];
 
 			if (!DmlBiasTensorDesc
 					.SetTensorRank(4, 4)
diff --git a/Engine/Plugins/Experimental/NNERuntimeRDG/Source/NNERuntimeRDG/Private/Dml/NNEDmlOperatorLRN.cpp b/Engine/Plugins/Experimental/NNERuntimeRDG/Source/NNERuntimeRDG/Private/Dml/NNEDmlOperatorLRN.cpp
--- a/Engine/Plugins/Experimental/NNERuntimeRDG/Source/NNERuntimeRDG/Private/Dml/NNEDmlOperatorLRN.cpp
+++ b/Engine/Plugins/Experimental/NNERuntimeRDG/Source/NNERuntimeRDG/Private/Dml/NNEDmlOperatorLRN.cpp
@@ -12,13 +12,24 @@ namespace UE::NNERuntimeRDG::Private::Dml
 //
 class FOperatorDmlLRN : public FOperatorDml
 {
+	int32 Size;
+	float Alpha;
+	float Beta;
+	float Bias;
+
 public:
 
+	//
+	//
+	//
 	static FOperatorDml* Create()
 	{
 		return new FOperatorDmlLRN();
 	}
 
+	//
+	//
+	//
 	static bool Validate(const NNE::FAttributeMap& AttributeMap, TConstArrayView<ENNETensorDataType> InputTypes, TConstArrayView<NNE::FSymbolicTensorShape> InputShapes)
 	{
 		//TODO
@@ -28,25 +39,36 @@ public:
 	//
 	//
 	//
-	virtual bool Initialize(IDMLDevice* Device, TArrayView<const NNE::Internal::FTensor> InputTensors, TArrayView<const NNE::Internal::FTensor> OutputTensors, const NNE::FAttributeMap& Attributes) override
+	virtual bool Initialize(TConstArrayView<NNE::FTensorDesc> Inputs, TConstArrayView<NNE::FTensorDesc> Outputs, const NNE::FAttributeMap& Attributes) override
 	{
-		check(InputTensors.Num() == 1);
-		check(OutputTensors.Num() == 1);
+		check(Inputs.Num() == 1);
+		check(Outputs.Num() == 1);
 
-		const NNE::Internal::FTensor& InputTensor = InputTensors[0];
-		const NNE::Internal::FTensor& OutputTensor = OutputTensors[0];
+		// Read attributes
+		Size = Attributes.GetValueOrDefault<int32>(TEXT("size"), 0);
+		Alpha = Attributes.GetValueOrDefault<float>(TEXT("alpha"), 0.0f);
+		Beta = Attributes.GetValueOrDefault<float>(TEXT("beta"), 0.0f);
+		Bias = Attributes.GetValueOrDefault<float>(TEXT("bias"), 0.0f);
 
-		if (InputTensor.GetShape().Rank() > 8)
-		{
-			UE_LOG(LogNNE, Warning, TEXT("InputTensor rank should be between 1 and 8, got:%d"), InputTensor.GetShape().Rank());
-			return false;
-		}
+		return true;
+	}
 
-		// Read attributes
-		const int32 Size = Attributes.GetValueOrDefault<int32>(TEXT("size"), 0);
-		const float Alpha = Attributes.GetValueOrDefault<float>(TEXT("alpha"), 0.0f);
-		const float Beta = Attributes.GetValueOrDefault<float>(TEXT("beta"), 0.0f);
-		const float Bias = Attributes.GetValueOrDefault<float>(TEXT("bias"), 0.0f);
+	//
+	//
+	//
+	virtual int PrepareOutputs(TConstArrayView<NNE::Internal::FTensorRef> InputTensors, TArrayView<NNE::Internal::FTensorRef> OutputTensors) const override
+	{
+		OutputTensors[0]->SetShape(InputTensors[0]->GetShape());
+		return 0;
+	}
+
+	//
+	//
+	//
+	virtual bool Create(IDMLDevice* Device, TConstArrayView<NNE::Internal::FTensorRef> InputTensors, TConstArrayView<NNE::Internal::FTensorRef> OutputTensors) override
+	{
+		const NNE::Internal::FTensor& InputTensor = *InputTensors[0];
+		const NNE::Internal::FTensor& OutputTensor = *OutputTensors[0];
 
 		FTensorDescDml	DmlInputTensorDesc;
 
diff --git a/Engine/Plugins/Experimental/NNERuntimeRDG/Source/NNERuntimeRDG/Private/Dml/NNEDmlOperatorLpNormalization.cpp b/Engine/Plugins/Experimental/NNERuntimeRDG/Source/NNERuntimeRDG/Private/Dml/NNEDmlOperatorLpNormalization.cpp
--- a/Engine/Plugins/Experimental/NNERuntimeRDG/Source/NNERuntimeRDG/Private/Dml/NNEDmlOperatorLpNormalization.cpp
+++ b/Engine/Plugins/Experimental/NNERuntimeRDG/Source/NNERuntimeRDG/Private/Dml/NNEDmlOperatorLpNormalization.cpp
@@ -14,13 +14,23 @@ class FOperatorDmlLpNormalization : public FOperatorDml
 {
 	static constexpr float DefaultEpsilon = 1e-5f;
 
+	int32	P;
+	uint32	Axis;
+
+
 public:
 
+	//
+	//
+	//
 	static FOperatorDml* Create()
 	{
 		return new FOperatorDmlLpNormalization();
 	}
 
+	//
+	//
+	//
 	static bool Validate(const NNE::FAttributeMap& AttributeMap, TConstArrayView<ENNETensorDataType> InputTypes, TConstArrayView<NNE::FSymbolicTensorShape> InputShapes)
 	{
 		//TODO
@@ -30,26 +40,37 @@ public:
 	//
 	//
 	//
-	virtual bool Initialize(IDMLDevice* Device, TArrayView<const NNE::Internal::FTensor> InputTensors, TArrayView<const NNE::Internal::FTensor> OutputTensors, const NNE::FAttributeMap& Attributes) override
+	virtual bool Initialize(TConstArrayView<NNE::FTensorDesc> Inputs, TConstArrayView<NNE::FTensorDesc> Outputs, const NNE::FAttributeMap& Attributes) override
 	{
-		check(InputTensors.Num() == 1);
-		check(OutputTensors.Num() == 1);
-
-		const NNE::Internal::FTensor& InputTensor = InputTensors[0];
-		const NNE::Internal::FTensor& OutputTensor = OutputTensors[0];
-
-		if (InputTensor.GetShape().Rank() > 8)
-		{
-			UE_LOG(LogNNE, Error, TEXT("InputTensor rank should be between 1 and 8, got:%d"), InputTensor.GetShape().Rank());
-			return false;
-		}
+		check(Inputs.Num() == 1);
+		check(Outputs.Num() == 1);
 
 		// Read attributes
-		const int32 P = Attributes.GetValueOrDefault<int32>(TEXT("p"), 1);
+		P = Attributes.GetValueOrDefault<int32>(TEXT("p"), 1);
 		check(P >= 1 && P <= 2);
 
-		int32 OnnxAxis = Attributes.GetValueOrDefault<int32>(TEXT("axis"), 0);
-		
+		Axis = Attributes.GetValueOrDefault<int32>(TEXT("axis"), 0);
+
+		return true;
+	}
+
+	//
+	//
+	//
+	virtual int PrepareOutputs(TConstArrayView<NNE::Internal::FTensorRef> InputTensors, TArrayView<NNE::Internal::FTensorRef> OutputTensors) const override
+	{
+		OutputTensors[0]->SetShape(InputTensors[0]->GetShape());
+		return 0;
+	}
+
+	//
+	//
+	//
+	virtual bool Create(IDMLDevice* Device, TConstArrayView<NNE::Internal::FTensorRef> InputTensors, TConstArrayView<NNE::Internal::FTensorRef> OutputTensors) override
+	{
+		const NNE::Internal::FTensor& InputTensor = *InputTensors[0];
+		const NNE::Internal::FTensor& OutputTensor = *OutputTensors[0];
+
 		FTensorDescDml	DmlInputTensorDesc;
 		
 		if (!DmlInputTensorDesc
@@ -70,13 +91,11 @@ public:
 			return false;
 		}
 
-		uint32 DmlAxis = GetDmlAxis(OnnxAxis, InputTensor.GetShape().Rank(), DmlInputTensorDesc.GetRank());
-
 		DML_LP_NORMALIZATION_OPERATOR_DESC	OpDesc{};
 
 		OpDesc.InputTensor = DmlInputTensorDesc.GetDmlDesc();
 		OpDesc.OutputTensor = DmlOutputTensorDesc.GetDmlDesc();
-		OpDesc.Axis = DmlAxis;
+		OpDesc.Axis = GetDmlAxis(Axis, InputTensor.GetShape().Rank(), DmlInputTensorDesc.GetRank());;
 		OpDesc.Epsilon = DefaultEpsilon;
 		OpDesc.P = P;
 
diff --git a/Engine/Plugins/Experimental/NNERuntimeRDG/Source/NNERuntimeRDG/Private/Dml/NNEDmlOperatorMatMul.cpp b/Engine/Plugins/Experimental/NNERuntimeRDG/Source/NNERuntimeRDG/Private/Dml/NNEDmlOperatorMatMul.cpp
--- a/Engine/Plugins/Experimental/NNERuntimeRDG/Source/NNERuntimeRDG/Private/Dml/NNEDmlOperatorMatMul.cpp
+++ b/Engine/Plugins/Experimental/NNERuntimeRDG/Source/NNERuntimeRDG/Private/Dml/NNEDmlOperatorMatMul.cpp
@@ -27,16 +27,74 @@ public:
 	//
 	//
 	//
-	virtual bool Initialize(IDMLDevice* Device, TArrayView<const NNE::Internal::FTensor> InputTensors, TArrayView<const NNE::Internal::FTensor> OutputTensors, const NNE::FAttributeMap& Attributes) override
+	virtual bool Initialize(TConstArrayView<NNE::FTensorDesc> Inputs, TConstArrayView<NNE::FTensorDesc> Outputs, const NNE::FAttributeMap& Attributes) override
 	{
-		check(InputTensors.Num() == 2);
+		check(Inputs.Num() == 2);
+		check(Outputs.Num() == 1);
+
+		return true;
+	}
+
+	//
+	//
+	//
+	virtual int PrepareOutputs(TConstArrayView<NNE::Internal::FTensorRef> InputTensors, TArrayView<NNE::Internal::FTensorRef> OutputTensors) const override
+	{
+		check(InputTensors.Num() >= 2 && InputTensors.Num() <= 3);
 		check(OutputTensors.Num() == 1);
 
-		using namespace UE::NNE;
+		const NNE::FTensorShape& InputA = InputTensors[0]->GetShape();
+		const NNE::FTensorShape& InputB = InputTensors[1]->GetShape();
+
+		if (InputA.Rank() < 2)
+		{
+			UE_LOG(LogNNE, Warning, TEXT("Matmul first input should be at least of rank 2"));
+			return -1;
+		}
+		
+		if (InputB.Rank() < 2)
+		{
+			UE_LOG(LogNNE, Warning, TEXT("Matmul second input should be at least of rank 2"));
+			return -1;
+		}
+
+		if (InputA.GetData()[InputA.Rank() - 1] != InputB.GetData()[InputB.Rank() - 2])
+		{
+			UE_LOG(LogNNE, Warning, TEXT("Matmul first input last dimension should be equal to second input last dimension"));
+			return -1;
+		}
+
+		const int32 OutputRank = FMath::Max(InputA.Rank(), InputB.Rank());
+		TArray<uint32> OutputShape;
+		OutputShape.SetNumUninitialized(OutputRank);
+
+		//Broadcast
+		for (int32 i = 0; i < OutputRank; ++i)
+		{
+			int32 AIndex = InputA.Rank() - 1 - i;
+			int32 BIndex = InputB.Rank() - 1 - i;
+			int32 AValue = AIndex >= 0 ? InputA.GetData()[AIndex] : 1;
+			int32 BValue = BIndex >= 0 ? InputB.GetData()[BIndex] : 1;
+			int32 OutputValue = FMath::Max(AValue, BValue);
+			OutputShape[OutputRank - 1 - i] = OutputValue;
+		}
+
+		//2D Mat
+		OutputShape[OutputRank - 2] = InputA.GetData()[InputA.Rank() - 2];
+		OutputShape[OutputRank - 1] = InputB.GetData()[InputB.Rank() - 1];
+
+		OutputTensors[0]->SetShape(NNE::FTensorShape::Make(OutputShape));
+		return 0;
+	}
 
-		const NNE::Internal::FTensor& ATensor = InputTensors[0];
-		const NNE::Internal::FTensor& BTensor = InputTensors[1];
-		const NNE::Internal::FTensor& OutTensor = OutputTensors[0];
+	//
+	//
+	//
+	virtual bool Create(IDMLDevice* Device, TConstArrayView<NNE::Internal::FTensorRef> InputTensors, TConstArrayView<NNE::Internal::FTensorRef> OutputTensors) override
+	{
+		const NNE::Internal::FTensor& ATensor = *InputTensors[0];
+		const NNE::Internal::FTensor& BTensor = *InputTensors[1];
+		const NNE::Internal::FTensor& OutTensor = *OutputTensors[0];
 
 		auto TrimMatrixStackLeadingOnes = [](TConstArrayView<uint32>&& Input) -> TConstArrayView<uint32>
 		{
diff --git a/Engine/Plugins/Experimental/NNERuntimeRDG/Source/NNERuntimeRDG/Private/Dml/NNEDmlOperatorMaxUnpool.cpp b/Engine/Plugins/Experimental/NNERuntimeRDG/Source/NNERuntimeRDG/Private/Dml/NNEDmlOperatorMaxUnpool.cpp
--- a/Engine/Plugins/Experimental/NNERuntimeRDG/Source/NNERuntimeRDG/Private/Dml/NNEDmlOperatorMaxUnpool.cpp
+++ b/Engine/Plugins/Experimental/NNERuntimeRDG/Source/NNERuntimeRDG/Private/Dml/NNEDmlOperatorMaxUnpool.cpp
@@ -8,16 +8,68 @@
 namespace UE::NNERuntimeRDG::Private::Dml
 {
 
+//
+//
+//
 class FOperatorDmlMaxUnpool : public FOperatorDml
 {
+	//
+	//
+	//
+	class FUnpoolingArgs : public FKernelArgs
+	{
+		Util::FSmallUIntArray	KernelShape;
+
+	public:
+
+		//
+		//
+		//
+		bool Init(const NNE::FAttributeMap& Attributes, const int32 InputShapeRank)
+		{
+			if (!FKernelArgs::Init(Attributes, InputShapeRank, /*bIsGlobalKernel*/ false, /*bIsTransposed*/ true))
+			{
+				return false;
+			}
+
+			if (!Util::GetArrayAttributeNoOverflow(Attributes.GetAttributeValue(TEXT("kernel_shape")), KernelShape))
+			{
+				UE_LOG(LogNNE, Error, TEXT("kernel_shape attribute cast led to overflow"));
+				return false;
+			}
+
+			if (KernelShape.IsEmpty())
+			{
+				UE_LOG(LogNNE, Error, TEXT("kernel_shape attribute is required for pooling operators"));
+				return false;
+			}
+
+			check(KernelShape.Num() == NumDimensions);
+
+			for (uint32 Dim = 0; Dim < NumDimensions; ++Dim)
+			{
+				WindowSize.Add(uint32(KernelShape[KernelShape.Num() - NumDimensions + Dim]));
+			}
+
+			return true;
+		}		
+	};
+
+	mutable FUnpoolingArgs Args;
 
 public:
 
+	//
+	//
+	//
 	static FOperatorDml* Create()
 	{
 		return new FOperatorDmlMaxUnpool();
 	}
 
+	//
+	//
+	//
 	static bool Validate(const NNE::FAttributeMap& AttributeMap, TConstArrayView<ENNETensorDataType> InputTypes, TConstArrayView<NNE::FSymbolicTensorShape> InputShapes)
 	{
 		//TODO
@@ -27,19 +79,67 @@ public:
 	//
 	//
 	//
-	virtual bool Initialize(IDMLDevice* Device, TArrayView<const NNE::Internal::FTensor> InputTensors, TArrayView<const NNE::Internal::FTensor> OutputTensors, const NNE::FAttributeMap& Attributes) override
+	virtual bool Initialize(TConstArrayView<NNE::FTensorDesc> Inputs, TConstArrayView<NNE::FTensorDesc> Outputs, const NNE::FAttributeMap& Attributes) override
 	{
-		check(InputTensors.Num() == 2 || InputTensors.Num() == 3);
-		check(OutputTensors.Num() == 1);
+		check(Inputs.Num() == 2 || Inputs.Num() == 3);
+		check(Outputs.Num() == 1);
 
-		if(InputTensors.Num() == 3)
+		if (Inputs.Num() == 3)
 		{
 			ConstantCPUInputs.Add(2);
 		}
 
-		const NNE::Internal::FTensor& InputTensor = InputTensors[0];
-		const NNE::Internal::FTensor& IndicesTensor = InputTensors[1];
-		const NNE::Internal::FTensor& OutputTensor = OutputTensors[0];
+		// DML required IndicesTensor to be in uint64 format, where ONNX allows int32 and int64
+		if (Inputs[1].GetDataType() != ENNETensorDataType::Int64)
+		{
+			UE_LOG(LogNNE, Error, TEXT("DML MaxUnpool requires UInt64, please use Int64 in your ONNX model"));
+			return false;
+		}
+
+		if (!Args.Init(Attributes, Inputs[0].GetShape().Rank()))
+		{
+			UE_LOG(LogNNE, Error, TEXT("DML MaxUnpool failed to initialize arguments"));
+			return false;
+		}
+
+		return true;
+	}
+
+	//
+	//
+	//
+	virtual int PrepareOutputs(TConstArrayView<NNE::Internal::FTensorRef> InputTensors, TArrayView<NNE::Internal::FTensorRef> OutputTensors) const override
+	{
+		Args.Evaluate(InputTensors[0]->GetShape().GetData());
+
+		if (InputTensors.Num() == 3)
+		{
+			TConstArrayView<int64>	Shape = InputTensors[2]->GetPreparedData<int64>();
+			Util::FSmallUIntArray	OutputShape;
+
+			for (int32 Idx = 0; Idx < Shape.Num(); ++Idx)
+			{
+				OutputShape.Add(static_cast<uint32>(std::clamp<int64>(Shape[Idx], INT32_MIN, INT32_MAX)));
+			}
+
+			OutputTensors[0]->SetShape(NNE::FTensorShape::Make(OutputShape));
+		}
+		else
+		{
+			OutputTensors[0]->SetShape(NNE::FTensorShape::Make(Args.GetOutputShape()));
+		}
+
+		return 0;
+	}
+
+	//
+	//
+	//
+	virtual bool Create(IDMLDevice* Device, TConstArrayView<NNE::Internal::FTensorRef> InputTensors, TConstArrayView<NNE::Internal::FTensorRef> OutputTensors) override
+	{
+		const NNE::Internal::FTensor& InputTensor = *InputTensors[0];
+		const NNE::Internal::FTensor& IndicesTensor = *InputTensors[1];
+		const NNE::Internal::FTensor& OutputTensor = *OutputTensors[0];
 
 		FTensorDescDml	DmlInputTensorDesc;
 		FTensorDescDml	DmlIndicesTensorDesc;
@@ -54,13 +154,6 @@ public:
 			return false;
 		}
 
-		// DML required IndicesTensor to be in uint64 format, where ONNX allows int32 and int64
-		if (IndicesTensor.GetDataType() != ENNETensorDataType::Int64)
-		{
-			UE_LOG(LogNNE, Error, TEXT("DML MaxUnpool requires UInt64, please use Int32 in your ONNX model"));
-			return false;
-		}
-
 		if (!DmlIndicesTensorDesc
 				.SetTensorRank(4, 4)
 				.SetFromTensor(IndicesTensor)
diff --git a/Engine/Plugins/Experimental/NNERuntimeRDG/Source/NNERuntimeRDG/Private/Dml/NNEDmlOperatorMeanVarianceNormalization.cpp b/Engine/Plugins/Experimental/NNERuntimeRDG/Source/NNERuntimeRDG/Private/Dml/NNEDmlOperatorMeanVarianceNormalization.cpp
--- a/Engine/Plugins/Experimental/NNERuntimeRDG/Source/NNERuntimeRDG/Private/Dml/NNEDmlOperatorMeanVarianceNormalization.cpp
+++ b/Engine/Plugins/Experimental/NNERuntimeRDG/Source/NNERuntimeRDG/Private/Dml/NNEDmlOperatorMeanVarianceNormalization.cpp
@@ -14,13 +14,23 @@ class FOperatorDmlMeanVarianceNormalization : public FOperatorDml
 {
 	static constexpr float DefaultEpsilon = 1e-5f;
 
+	Util::FSmallUIntArray	Axes;
+	int32					AcrossChannels;
+	int32					NormalizeVariance;
+
 public:
 
+	//
+	//
+	//
 	static FOperatorDml* Create()
 	{
 		return new FOperatorDmlMeanVarianceNormalization();
 	}
 
+	//
+	//
+	//
 	static bool Validate(const NNE::FAttributeMap& AttributeMap, TConstArrayView<ENNETensorDataType> InputTypes, TConstArrayView<NNE::FSymbolicTensorShape> InputShapes)
 	{
 		//TODO
@@ -30,23 +40,14 @@ public:
 	//
 	//
 	//
-	virtual bool Initialize(IDMLDevice* Device, TArrayView<const NNE::Internal::FTensor> InputTensors, TArrayView<const NNE::Internal::FTensor> OutputTensors, const NNE::FAttributeMap& Attributes) override
+	virtual bool Initialize(TConstArrayView<NNE::FTensorDesc> Inputs, TConstArrayView<NNE::FTensorDesc> Outputs, const NNE::FAttributeMap& Attributes) override
 	{
-		check(InputTensors.Num() == 1);
-		check(OutputTensors.Num() == 1);
-
-		const NNE::Internal::FTensor& InputTensor = InputTensors[0];
-		const NNE::Internal::FTensor& OutputTensor = OutputTensors[0];
-
-		if (InputTensor.GetShape().Rank() > 8)
-		{
-			UE_LOG(LogNNE, Warning, TEXT("InputTensor rank should be between 1 and 8, got:%d"), InputTensor.GetShape().Rank());
-			return false;
-		}
+		check(Inputs.Num() == 1);
+		check(Outputs.Num() == 1);
 
 		// Read attributes
-		const int32 AcrossChannels = Attributes.GetValueOrDefault<int32>(TEXT("across_channels"), 0);
-		const int32 NormalizeVariance = Attributes.GetValueOrDefault<int32>(TEXT("normalize_variance"), 1);
+		AcrossChannels = Attributes.GetValueOrDefault<int32>(TEXT("across_channels"), 0);
+		NormalizeVariance = Attributes.GetValueOrDefault<int32>(TEXT("normalize_variance"), 1);
 
 		TArray<int32>	OnnxAxes;
 
@@ -58,7 +59,7 @@ public:
 		}
 		else
 		{
-			constexpr int32 CrossChannelAxes[] = { 0, 1, 2, 3};
+			constexpr int32 CrossChannelAxes[] = { 0, 1, 2, 3 };
 			constexpr int32 NonChannelAxes[] = { 0, 2, 3 };
 
 			if (AcrossChannels)
@@ -71,6 +72,34 @@ public:
 			}
 		}
 
+		SetDmlAxesFromOnnx(Axes, Inputs[0].GetShape().Rank(), OnnxAxes);
+
+		return true;
+	}
+
+	//
+	//
+	//
+	virtual int PrepareOutputs(TConstArrayView<NNE::Internal::FTensorRef> InputTensors, TArrayView<NNE::Internal::FTensorRef> OutputTensors) const override
+	{
+		OutputTensors[0]->SetShape(InputTensors[0]->GetShape());
+		return 0;
+	}
+
+	//
+	//
+	//
+	virtual bool Create(IDMLDevice* Device, TConstArrayView<NNE::Internal::FTensorRef> InputTensors, TConstArrayView<NNE::Internal::FTensorRef> OutputTensors) override
+	{
+		const NNE::Internal::FTensor& InputTensor = *InputTensors[0];
+		const NNE::Internal::FTensor& OutputTensor = *OutputTensors[0];
+
+		if (InputTensor.GetShape().Rank() > 8)
+		{
+			UE_LOG(LogNNE, Warning, TEXT("InputTensor rank should be between 1 and 8, got:%d"), InputTensor.GetShape().Rank());
+			return false;
+		}
+
 		FTensorDescDml	DmlInputTensorDesc;
 		
 		if (!DmlInputTensorDesc
@@ -90,18 +119,15 @@ public:
 			UE_LOG(LogNNE, Error, TEXT("Failed to initialize tensor(s) for DML inference"));
 			return false;
 		}
-
-		Util::FSmallUIntArray DmlAxes;
-		SetDmlAxesFromOnnx(DmlAxes, DmlInputTensorDesc.GetRank(), OnnxAxes);
-
+		
 		DML_MEAN_VARIANCE_NORMALIZATION1_OPERATOR_DESC	OpDesc{};
 
 		OpDesc.InputTensor = DmlInputTensorDesc.GetDmlDesc();
 		OpDesc.ScaleTensor = nullptr;
 		OpDesc.BiasTensor = nullptr;
 		OpDesc.OutputTensor = DmlOutputTensorDesc.GetDmlDesc();
-		OpDesc.AxisCount = DmlAxes.Num();
-		OpDesc.Axes = DmlAxes.GetData();
+		OpDesc.AxisCount = Axes.Num();
+		OpDesc.Axes = Axes.GetData();
 		OpDesc.NormalizeVariance = NormalizeVariance;
 		OpDesc.Epsilon = DefaultEpsilon;
 		OpDesc.FusedActivation = nullptr;
diff --git a/Engine/Plugins/Experimental/NNERuntimeRDG/Source/NNERuntimeRDG/Private/Dml/NNEDmlOperatorPad.cpp b/Engine/Plugins/Experimental/NNERuntimeRDG/Source/NNERuntimeRDG/Private/Dml/NNEDmlOperatorPad.cpp
--- a/Engine/Plugins/Experimental/NNERuntimeRDG/Source/NNERuntimeRDG/Private/Dml/NNEDmlOperatorPad.cpp
+++ b/Engine/Plugins/Experimental/NNERuntimeRDG/Source/NNERuntimeRDG/Private/Dml/NNEDmlOperatorPad.cpp
@@ -5,10 +5,18 @@
 
 namespace UE::NNERuntimeRDG::Private::Dml
 {
-
+//
+//
+//
 class FOperatorDmlPad : public FOperatorDml
 {
+	float					Value;
+	DML_PADDING_MODE		Mode;
+	mutable TArray<int32>	Pads;
 
+	//
+	//
+	//
 	static DML_PADDING_MODE ModeFromString(FStringView StringVal)
 	{
 		if (FCString::Stricmp(StringVal.GetData(), TEXT("CONSTANT")) == 0)
@@ -31,11 +39,17 @@ class FOperatorDmlPad : public FOperatorDml
 
 public:
 
+	//
+	//
+	//
 	static FOperatorDml* Create()
 	{
 		return new FOperatorDmlPad();
 	}
 
+	//
+	//
+	//
 	static bool Validate(const NNE::FAttributeMap& AttributeMap, TConstArrayView<ENNETensorDataType> InputTypes, TConstArrayView<NNE::FSymbolicTensorShape> InputShapes)
 	{
 		//TODO
@@ -45,30 +59,17 @@ public:
 	//
 	//
 	//
-	virtual bool Initialize(IDMLDevice* Device, TArrayView<const NNE::Internal::FTensor> InputTensors, TArrayView<const NNE::Internal::FTensor> OutputTensors, const NNE::FAttributeMap& Attributes) override
+	virtual bool Initialize(TConstArrayView<NNE::FTensorDesc> Inputs, TConstArrayView<NNE::FTensorDesc> Outputs, const NNE::FAttributeMap& Attributes) override
 	{
-		check(InputTensors.Num() >= 1);
-		check(OutputTensors.Num() == 1);
-
-		const NNE::Internal::FTensor& InputTensor = InputTensors[0];
-		const NNE::Internal::FTensor& OutputTensor = OutputTensors[0];
+		check(Inputs.Num() >= 1);
+		check(Outputs.Num() == 1);
 
 		// Read attributes
-		float				Value;
-		DML_PADDING_MODE	Mode;
-		TArray<int32>		Pads;
-
 		Value = Attributes.GetValueOrDefault<float>(TEXT("value"), 0.0f);
 		Mode = ModeFromString(Attributes.GetValue<FString>(TEXT("mode")));
 
-		if (InputTensors.Num() >= 2)
+		if (Inputs.Num() >= 2)
 		{
-			if (!InputTensors[1].HasPreparedData())
-			{
-				UE_LOG(LogNNE, Error, TEXT("pads is only supported as an attribute or a constant tensor, it is here a variable tensor of name %s."), *InputTensors[1].GetName());
-				return false;
-			}
-			Pads.Append(InputTensors[1].GetPreparedData<int64>());
 			ConstantCPUInputs.Add(1);
 		}
 		else
@@ -76,12 +77,71 @@ public:
 			Pads = Attributes.GetValue<TArray<int32>>(TEXT("pads"));
 		}
 
-		if (InputTensor.GetShape().Rank() * 2 != Pads.Num())
+		return true;
+	}
+
+	//
+	//
+	//
+	virtual int PrepareOutputs(TConstArrayView<NNE::Internal::FTensorRef> InputTensors, TArrayView<NNE::Internal::FTensorRef> OutputTensors) const override
+	{
+		if (InputTensors.Num() >= 2)
+		{
+			if (!InputTensors[1]->HasPreparedData())
+			{
+				UE_LOG(LogNNE, Error, TEXT("pads is only supported as an attribute or a constant tensor, it is here a variable tensor of name %s."), *InputTensors[1]->GetName());
+				return -1;
+			}
+
+			Pads.Append(InputTensors[1]->GetPreparedData<int64>());
+		}
+
+		if (Pads.IsEmpty())
 		{
-			UE_LOG(LogNNE, Error, TEXT("pads attribute lenght (%d) should be twice the rank of input X (%d)."), Pads.Num(), InputTensor.GetShape().Rank());
+			UE_LOG(LogNNE, Error, TEXT("pads are empty"));
+			return -1;
+		}
+
+		TConstArrayView<uint32> InputShape = InputTensors[0]->GetShape().GetData();
+
+		if (InputShape.Num() * 2 != Pads.Num())
+		{
+			UE_LOG(LogNNE, Error, TEXT("pads attribute length (%d) should be twice the rank of input X (%d)."), Pads.Num(), InputShape.Num());
 			return false;
 		}
-		
+
+		TArray<uint32> OutputShape;
+		for (int32 i = 0; i < InputShape.Num(); ++i)
+		{
+			int32 PrePad = Pads[i];
+			int32 PostPad = Pads[i + InputShape.Num()];
+			int32 OutputDim = PrePad + InputShape[i] + PostPad;
+
+			if (OutputDim < 1)
+			{
+				UE_LOG(
+					LogNNE, Warning, TEXT("Pads cannot reduce dimension below 1, but would for tensor (name:%s) at rank %d of size %d with prepad %d and postpad %d."), 
+					*InputTensors[0]->GetName(), i, InputShape[i], PrePad, PostPad);
+
+				return -1;
+			}
+
+			OutputShape.Emplace(OutputDim);
+		}
+
+		OutputTensors[0]->SetShape(NNE::FTensorShape::Make(OutputShape));
+
+		return 0;
+	}
+
+	//
+	//
+	//
+	virtual bool Create(IDMLDevice* Device, TConstArrayView<NNE::Internal::FTensorRef> InputTensors, TConstArrayView<NNE::Internal::FTensorRef> OutputTensors) override
+	{
+		const NNE::Internal::FTensor& InputTensor = *InputTensors[0];
+		const NNE::Internal::FTensor& OutputTensor = *OutputTensors[0];
+
 		// Initialize tensor descriptors
 		FTensorDescDml	DmlInputTensorDesc;
 		FTensorDescDml	DmlOutputTensorDesc;
diff --git a/Engine/Plugins/Experimental/NNERuntimeRDG/Source/NNERuntimeRDG/Private/Dml/NNEDmlOperatorPool.cpp b/Engine/Plugins/Experimental/NNERuntimeRDG/Source/NNERuntimeRDG/Private/Dml/NNEDmlOperatorPool.cpp
--- a/Engine/Plugins/Experimental/NNERuntimeRDG/Source/NNERuntimeRDG/Private/Dml/NNEDmlOperatorPool.cpp
+++ b/Engine/Plugins/Experimental/NNERuntimeRDG/Source/NNERuntimeRDG/Private/Dml/NNEDmlOperatorPool.cpp
@@ -11,20 +11,100 @@ namespace UE::NNERuntimeRDG::Private::Dml
 {
 
 /**
-	* MaxPool
-	*/
+ *  Pooling operators (local and global)
+ */
 template <typename TDmlPoolOpDesc, DML_OPERATOR_TYPE DmlOpType, bool UseGlobalPooling>
 class FOperatorDmlPool : public FOperatorDml
 {
 	using FIntArray = TArray<int32>;
 
+	//
+	//
+	//
+	class FPoolingArgs : public FKernelArgs
+	{
+		Util::FSmallUIntArray	KernelShape;
+
+	public:
+
+		//
+		//
+		//
+		bool Init(const NNE::FAttributeMap& Attributes, const int32 InputShapeRank, bool bInIsGlobalKernel)
+		{
+			if (!FKernelArgs::Init(Attributes, InputShapeRank, UseGlobalPooling, /*bIsTransposed*/ false))
+			{
+				return false;
+			}
+
+			if (!bIsGlobalKernel)
+			{
+				if (!Util::GetArrayAttributeNoOverflow(Attributes.GetAttributeValue(TEXT("kernel_shape")), KernelShape))
+				{
+					UE_LOG(LogNNE, Error, TEXT("kernel_shape attribute cast led to overflow"));
+					return false;
+				}
+
+				if (KernelShape.IsEmpty())
+				{
+					UE_LOG(LogNNE, Error, TEXT("kernel_shape attribute is required for pooling operators"));
+					return false;
+				}
+				
+				check(KernelShape.Num() == NumDimensions);
+
+				for (uint32 Dim = 0; Dim < NumDimensions; ++Dim)
+				{
+					WindowSize.Add(uint32(KernelShape[KernelShape.Num() - NumDimensions + Dim]));
+				}
+			}
+
+			return true;
+		}
+
+		//
+		//
+		//
+		bool Evaluate(TConstArrayView<uint32> InputShape)
+		{
+			if (!bIsGlobalKernel)
+			{
+				FKernelArgs::Evaluate(InputShape, KernelPadding(InputShape, WindowSize, Dilations, Strides));
+			}
+			else
+			{
+				FKernelArgs::Evaluate(InputShape);
+			}
+
+			return true;
+		}
+
+		//
+		//
+		//
+		TConstArrayView<uint32> GetWindowSize() const
+		{
+			return WindowSize;
+		}
+	};
+
+	mutable FPoolingArgs	Args;
+	int32					IncludePadding;
+	uint32					P;
+
 public:
 
+	//
+	//
+	//
 	static FOperatorDml* Create()
 	{
 		return new FOperatorDmlPool<TDmlPoolOpDesc, DmlOpType, UseGlobalPooling>();
 	}
 
+	//
+	//
+	//
 	static bool Validate(const NNE::FAttributeMap& AttributeMap, TConstArrayView<ENNETensorDataType> InputTypes, TConstArrayView<NNE::FSymbolicTensorShape> InputShapes)
 	{
 		//TODO
@@ -34,109 +114,95 @@ public:
 	//
 	//
 	//
-	virtual bool Initialize(IDMLDevice* Device, TArrayView<const NNE::Internal::FTensor> InputTensors, TArrayView<const NNE::Internal::FTensor> OutputTensors, const NNE::FAttributeMap& Attributes) override
+	virtual bool Initialize(TConstArrayView<NNE::FTensorDesc> Inputs, TConstArrayView<NNE::FTensorDesc> Outputs, const NNE::FAttributeMap& Attributes) override
 	{
-		//TODO: int64 attributes
-		check(InputTensors.Num() == 1);
-		TConstArrayView<uint32> InputShape = InputTensors[0].GetShape().GetData();
-		check(InputShape.Num() >= 2);
-		check(OutputTensors.Num() == 1 || OutputTensors.Num() == 2);
+		// TODO: int64 attributes
+		check(Inputs.Num() == 1);
+		check(Outputs.Num() == 1 || Outputs.Num() == 2);
 
-		const uint32 NumSpatialDimensions = InputShape.Num() - NonspatialDimensionCount;
-		if (NumSpatialDimensions != 2 && NumSpatialDimensions != 3)
+		// Attribute storage_order is not supported
+		const int32 StorageOrder = Attributes.GetValueOrDefault<int32>(TEXT("storage_order"), 0);
+		if (StorageOrder != 0)
 		{
-			UE_LOG(LogNNE, Error, TEXT("Number of spatial dimensions must be in range [2, 3] for DML inference."));
+			UE_LOG(LogNNE, Error, TEXT("storage_order != 0 is not supported for DML inference"));
 			return false;
 		}
 
-		Util::FSmallUIntArray StartPadding, EndPadding, KernelShape, Strides, Dilations;
+		const int32 InputShapeRank = Inputs[0].GetShape().Rank();
+		check(InputShapeRank >= 2);
 
-		Strides.Init(1u, NumSpatialDimensions);
-		if constexpr(!UseGlobalPooling)
+		if (!Args.Init(Attributes, InputShapeRank, UseGlobalPooling))
 		{
-			if (!Util::GetArrayAttributeNoOverflow(Attributes.GetAttributeValue(TEXT("strides")), Strides, TConstArrayView<uint32>(Strides)))
-			{
-				UE_LOG(LogNNE, Error, TEXT("Strides attribute cast led to overflow"));
-				return false;
-			}
+			return false;
 		}
-		check(Strides.Num() == NumSpatialDimensions);
 
-		Dilations.Init(1u, NumSpatialDimensions);
-		if constexpr(!UseGlobalPooling)
+		if constexpr (DmlOpType == DML_OPERATOR_AVERAGE_POOLING)
 		{
-			if (!Util::GetArrayAttributeNoOverflow(Attributes.GetAttributeValue(TEXT("dilations")), Dilations, TConstArrayView<uint32>(Dilations)))
-			{
-				UE_LOG(LogNNE, Error, TEXT("Dilations attribute cast led to overflow"));
-				return false;
-			}
+			IncludePadding = Attributes.GetValueOrDefault(TEXT("count_include_pad"), 0);
 		}
-		check(Dilations.Num() == NumSpatialDimensions);
-		
-		if constexpr(UseGlobalPooling)
+
+		if constexpr (DmlOpType == DML_OPERATOR_LP_POOLING)
 		{
-			KernelShape = Util::FSmallUIntArray{ InputShape.RightChop(NonspatialDimensionCount) };
+			P = (uint32) Attributes.GetValueOrDefault<int>(TEXT("p"), 2);
 		}
-		else
+
+		return true;
+	}
+
+	//
+	//
+	//
+	virtual int PrepareOutputs(TConstArrayView<NNE::Internal::FTensorRef> InputTensors, TArrayView<NNE::Internal::FTensorRef> OutputTensors) const override
+	{
+		TConstArrayView<uint32> InputShape = InputTensors[0]->GetShape().GetData();
+
+		Args.Evaluate(InputShape);
+
+		// Compute output shape(s)
 		{
-			if (!Util::GetArrayAttributeNoOverflow(Attributes.GetAttributeValue(TEXT("kernel_shape")), KernelShape))
+			TConstArrayView<uint32> OutputShape = Args.GetOutputShape();
+
+			for (auto& OutputTensor : OutputTensors)
 			{
-				UE_LOG(LogNNE, Error, TEXT("kernel_shape attribute cast led to overflow"));
-				return false;
+				OutputTensor->SetShape(NNE::FTensorShape::Make(OutputShape));
 			}
 		}
-		
-		if (KernelShape.IsEmpty())
-		{
-			UE_LOG(LogNNE, Error, TEXT("kernel_shape attribute is required for pooling operators"));
-			return false;
-		}
-		check(KernelShape.Num() == NumSpatialDimensions);
 
-		if constexpr(UseGlobalPooling)
-		{
-			StartPadding.Init(0u, NumSpatialDimensions);
-			EndPadding.Init(0u, NumSpatialDimensions);
-		}
-		else
-		{
-			ComputeStartEndPaddings(
-				InputShape,
-				Attributes,
-				StartPadding,
-				EndPadding,
-				KernelPadding(InputShape, KernelShape, Dilations, Strides)
-			);
-		}
-		
+		// Check output shape
 		for (int Idx = 0; Idx < OutputTensors.Num(); ++Idx)
 		{
-			check(OutputTensors[Idx].GetShape().Rank() == InputShape.Num());
-			check(OutputTensors[Idx].GetShape().GetData()[0] == InputShape[0]);
-			check(OutputTensors[Idx].GetShape().GetData()[1] == InputShape[1]);
+			check(OutputTensors[Idx]->GetShape().Rank() == InputShape.Num());
+			check(OutputTensors[Idx]->GetShape().GetData()[0] == InputShape[0]);
+			check(OutputTensors[Idx]->GetShape().GetData()[1] == InputShape[1]);
 
-			for (uint32 Dim = 0; Dim < NumSpatialDimensions; ++Dim) {
-				uint32 PaddedSize = InputShape[Dim + NonspatialDimensionCount] + StartPadding[Dim] + EndPadding[Dim];
-				check(OutputTensors[Idx].GetShape().GetData()[Dim + NonspatialDimensionCount] == (PaddedSize - KernelShape[Dim]) / Strides[Dim] + 1);
+			TConstArrayView<uint32> StartPadding = Args.GetStartPadding();
+			TConstArrayView<uint32> EndPadding = Args.GetEndPadding();
+			TConstArrayView<uint32> Strides = Args.GetStrides();
+			TConstArrayView<uint32> WindowSize = Args.GetWindowSize();
+			TConstArrayView<uint32> OutputShape = OutputTensors[Idx]->GetShape().GetData();
+
+			for (uint32 Dim = 0; Dim < Args.GetNumDimensions(); ++Dim)
+			{
+				const uint32 PaddedSize = InputShape[Dim + NonspatialDimensionCount] + StartPadding[Dim] + EndPadding[Dim];
+				check(OutputShape[Dim + NonspatialDimensionCount] == (PaddedSize - WindowSize[Dim]) / Strides[Dim] + 1);
 			}
 		}
 
-		//storage_order is not supported
-		int32 StorageOrder =
-			Attributes.GetValueOrDefault<int32>(
-				TEXT("storage_order"),
-				0);
-		if (StorageOrder != 0)
-		{
-			UE_LOG(LogNNE, Error, TEXT("storage_order != 0 is not supported for DML inference"));
-			return false;
-		}
+		return 0;
+	}
+
+	//
+	//
+	//
+	virtual bool Create(IDMLDevice* Device, TConstArrayView<NNE::Internal::FTensorRef> InputTensors, TConstArrayView<NNE::Internal::FTensorRef> OutputTensors) override
+	{
+		TConstArrayView<uint32> InputShape = InputTensors[0]->GetShape().GetData();
 
 		FTensorDescDml DmlInputTensorDesc;
 
 		if (!DmlInputTensorDesc
 				.SetTensorRank(4, 5)
-				.SetFromTensor(InputTensors[0])
+				.SetFromTensor(*InputTensors[0])
 				.Validate())
 		{
 			UE_LOG(LogNNE, Error, TEXT("Failed to initialize pooling operator's input tensor for DML inference"));
@@ -147,7 +213,7 @@ public:
 
 		if (!DmlOutputTensorDesc
 				.SetTensorRank(4, 5)
-				.SetFromTensor(OutputTensors[0])
+				.SetFromTensor(*OutputTensors[0])
 				.Validate())
 		{
 			UE_LOG(LogNNE, Error, TEXT("Failed to initialize pooling operator's output tensor for DML inference"));
@@ -158,23 +224,26 @@ public:
 		{
 			PoolingDesc.InputTensor = DmlInputTensorDesc.GetDmlDesc();
 			PoolingDesc.OutputTensor = DmlOutputTensorDesc.GetDmlDesc();
-			PoolingDesc.DimensionCount = NumSpatialDimensions;
-			PoolingDesc.WindowSize = KernelShape.GetData();
-			PoolingDesc.Strides = Strides.GetData();
-			PoolingDesc.StartPadding = StartPadding.GetData();
-			PoolingDesc.EndPadding = EndPadding.GetData();
+			PoolingDesc.DimensionCount = Args.GetNumDimensions();
+			PoolingDesc.WindowSize = Args.GetWindowSize().GetData();
+			PoolingDesc.Strides = Args.GetStrides().GetData();
+			PoolingDesc.StartPadding = Args.GetStartPadding().GetData();
+			PoolingDesc.EndPadding = Args.GetEndPadding().GetData();
 		};
 
 		const bool bHasDilations =
 			Algo::AnyOf(
-				Dilations,
-				[](auto Dim) {return Dim != 1u; }
+				Args.GetDilations(),
+				[](auto Dim) 
+				{
+					return Dim != 1u; 
+				}
 		);
+
 		const bool bHasOutputIndices = OutputTensors.Num() > 1;
 
 		TDmlPoolOpDesc DmlPoolOpDesc = {};
-		FillPoolingDesc(DmlPoolOpDesc);
-
+		
 		if (DmlOpType == DML_OPERATOR_MAX_POOLING && (bHasOutputIndices || bHasDilations))
 		{
 			DML_MAX_POOLING2_OPERATOR_DESC DmlMaxPoolOpDesc = {};
@@ -185,7 +254,7 @@ public:
 				
 				if (!DmlIndicesTensorDesc
 						.SetTensorRank(4, 5)
-						.SetFromTensor(OutputTensors[1])
+						.SetFromTensor(*OutputTensors[1])
 						.SetDataType(ENNETensorDataType::UInt64)
 						.Validate())
 				{
@@ -196,27 +265,33 @@ public:
 				DmlMaxPoolOpDesc.OutputIndicesTensor = DmlIndicesTensorDesc.GetDmlDesc();
 			}
 
-			DmlMaxPoolOpDesc.Dilations = Dilations.GetData();
+			DmlMaxPoolOpDesc.Dilations = Args.GetDilations().GetData();
 			FillPoolingDesc(DmlMaxPoolOpDesc);
 
 			return CreateOperator(Device, DML_OPERATOR_DESC{ DML_OPERATOR_MAX_POOLING2, &DmlMaxPoolOpDesc });
 		}
+		else
+		{
+			FillPoolingDesc(DmlPoolOpDesc);
+		}
 
 		if constexpr (DmlOpType == DML_OPERATOR_AVERAGE_POOLING)
 		{
-			int32 IncludePadding = Attributes.GetValueOrDefault(TEXT("count_include_pad"), 0);
 			DmlPoolOpDesc.IncludePadding = static_cast<BOOL>(IncludePadding);
 		}
 
 		if constexpr (DmlOpType == DML_OPERATOR_LP_POOLING)
 		{
-			DmlPoolOpDesc.P = (UINT) Attributes.GetValueOrDefault<int>(TEXT("p"), 2);
+			DmlPoolOpDesc.P = P;
 		}
 			
 		return CreateOperator(Device, DML_OPERATOR_DESC{ DmlOpType, &DmlPoolOpDesc });
 	}
 };
 
+//
+//
+//
 #define NNE_DML_REGISTER_POOLING_OP(OpName, DmlPrefix, UseGlobalPooling) \
 struct FDmlOperator##OpName##Registrator \
 { \
@@ -228,7 +303,9 @@ struct FDmlOperator##OpName##Registrator \
 \
 static FDmlOperator##OpName##Registrator RegisterDmlOperator##OpName;
 
+//
 // Register pooling operator on Module startup
+//
 NNE_DML_REGISTER_POOLING_OP(MaxPool, MAX, false)
 NNE_DML_REGISTER_POOLING_OP(GlobalMaxPool, MAX, true)
 NNE_DML_REGISTER_POOLING_OP(AveragePool, AVERAGE, false)
diff --git a/Engine/Plugins/Experimental/NNERuntimeRDG/Source/NNERuntimeRDG/Private/Dml/NNEDmlOperatorReduce.cpp b/Engine/Plugins/Experimental/NNERuntimeRDG/Source/NNERuntimeRDG/Private/Dml/NNEDmlOperatorReduce.cpp
--- a/Engine/Plugins/Experimental/NNERuntimeRDG/Source/NNERuntimeRDG/Private/Dml/NNEDmlOperatorReduce.cpp
+++ b/Engine/Plugins/Experimental/NNERuntimeRDG/Source/NNERuntimeRDG/Private/Dml/NNEDmlOperatorReduce.cpp
@@ -23,8 +23,10 @@ namespace UE::NNERuntimeRDG::Private::Dml
 template<DML_REDUCE_FUNCTION ReduceFunc>
 class FOperatorDmlReduce : public FOperatorDml
 {
-
-	inline void HandleEmptyAxes(TArray<int32>& Axes, int32 Rank)
+	//
+	//
+	//
+	inline static void HandleEmptyAxes(TArray<int32>& Axes, int32 Rank)
 	{
 		if (Axes.IsEmpty())
 		{
@@ -37,13 +39,25 @@ class FOperatorDmlReduce : public FOperatorDml
 		}
 	}
 
+	Util::FSmallUIntArray			Axes;
+	mutable Util::FSmallUIntArray	ReducedDims;
+	mutable Util::FSmallArray<bool>	IsReducedDims;
+	int32							KeepDims;
+	DML_AXIS_DIRECTION				AxisDirection{ DML_AXIS_DIRECTION_INCREASING };
+
 public:
 
+	//
+	//
+	//
 	static FOperatorDml* Create()
 	{
 		return new FOperatorDmlReduce<ReduceFunc>();
 	}
 
+	//
+	//
+	//
 	static bool Validate(const NNE::FAttributeMap& AttributeMap, TConstArrayView<ENNETensorDataType> InputTypes, TConstArrayView<NNE::FSymbolicTensorShape> InputShapes)
 	{
 		//TODO
@@ -53,81 +67,132 @@ public:
 	//
 	//
 	//
-	virtual bool Initialize(IDMLDevice* Device, TArrayView<const NNE::Internal::FTensor> InputTensors, TArrayView<const NNE::Internal::FTensor> OutputTensors, const NNE::FAttributeMap& Attributes) override
+	virtual bool Initialize(TConstArrayView<NNE::FTensorDesc> Inputs, TConstArrayView<NNE::FTensorDesc> Outputs, const NNE::FAttributeMap& Attributes) override
 	{
-		checkf(InputTensors.Num() == 1, TEXT("Dml Reduce op supports only 1 input"));
-		check(OutputTensors.Num() == 1);
+		checkf(Inputs.Num() == 1, TEXT("Dml Reduce op supports only 1 input"));
+		check(Outputs.Num() == 1);
 
-		const int32 KeepDims = Attributes.GetValueOrDefault<int32>(TEXT("keepdims"), 1);
-		
-		TArray<int32>	Axes;
+		KeepDims = Attributes.GetValueOrDefault<int32>(TEXT("keepdims"), 1);
+
+		if constexpr (ReduceFunc == DML_REDUCE_FUNCTION_ARGMAX || ReduceFunc == DML_REDUCE_FUNCTION_ARGMIN)
+		{
+			AxisDirection = (DML_AXIS_DIRECTION) Attributes.GetValueOrDefault<int32>(TEXT("select_last_index"), 0);
+		}
+
+		TArray<int32>	OnnxAxes;
 
 		if constexpr (ReduceFunc == DML_REDUCE_FUNCTION_ARGMAX || ReduceFunc == DML_REDUCE_FUNCTION_ARGMIN)
 		{
-			Axes.Add(Attributes.GetValueOrDefault<int32>(TEXT("axis"), 0));
+			OnnxAxes.Add(Attributes.GetValueOrDefault<int32>(TEXT("axis"), 0));
 		}
 		else
 		{
 			const FNNEAttributeValue* AxesAttr = Attributes.GetAttributeValue(TEXT("axes"));
 			if (AxesAttr)
 			{
-				Axes = AxesAttr->GetValue<TArray<int32>>();
+				OnnxAxes = AxesAttr->GetValue<TArray<int32>>();
 			}
 		}
 
-		HandleNegativeAxes(Axes, InputTensors[0].GetShape().Rank());
-		HandleEmptyAxes(Axes, InputTensors[0].GetShape().Rank());
+		HandleNegativeAxes(OnnxAxes, Inputs[0].GetShape().Rank());
+		HandleEmptyAxes(OnnxAxes, Inputs[0].GetShape().Rank());
 
+		const int32 InputRank = Inputs[0].GetShape().Rank();
+
+		for (int32& Dim : OnnxAxes)
+		{
+			checkf(Dim < InputRank, TEXT("Index out of bounds for Reduce axis"));
+			Axes.Add(Dim);
+		}
+
+		return true;
+	}
+
+	//
+	//
+	//
+	virtual int PrepareOutputs(TConstArrayView<NNE::Internal::FTensorRef> InputTensors, TArrayView<NNE::Internal::FTensorRef> OutputTensors) const override
+	{
+		TConstArrayView<uint32> InputShape = InputTensors[0]->GetShape().GetData();
+		Util::FSmallUIntArray	OutputShape;
+
+		ReducedDims.Reset();
+		ReducedDims.Append(InputShape);
+
+		for (const uint32& Dim : Axes)
+		{
+			checkf(Dim < (uint32) ReducedDims.Num(), TEXT("Index out of bounds for Reduce axis"));
+			ReducedDims[Dim] = 1;
+		}
+
+		if (!KeepDims)
+		{
+			IsReducedDims.SetNumZeroed(ReducedDims.Num());
+			for (const uint32& Dim : Axes)
+			{
+				IsReducedDims[Dim] = true;
+			}
+
+			for (int32 Idx = 0; Idx < ReducedDims.Num(); ++Idx)
+			{
+				if (!IsReducedDims[Idx])
+				{
+					OutputShape.Add(ReducedDims[Idx]);
+				}
+			}			
+		}
+		else
+		{
+			OutputShape.Append(ReducedDims);
+		}
+		
+		OutputTensors[0]->SetShape(NNE::FTensorShape::Make(OutputShape));
+
+		return 0;
+	}
+
+	//
+	//
+	//
+	virtual bool Create(IDMLDevice* Device, TConstArrayView<NNE::Internal::FTensorRef> InputTensors, TConstArrayView<NNE::Internal::FTensorRef> OutputTensors) override
+	{
 		FTensorDescDml DmlInputTensorDesc;
 		
 		if (!DmlInputTensorDesc
-				.SetFromTensor(InputTensors[0])
+				.SetFromTensor(*InputTensors[0])
 				.Validate())
 		{
 			UE_LOG(LogNNE, Warning, TEXT("Failed to initialize Reduce input tensor for DML inference"));
 			return false;
 		}
 
-		Util::FSmallUIntArray	DmlAxes;
-		Util::FSmallUIntArray	ReducedDims;
-
-		ReducedDims.SetNum(DmlInputTensorDesc.GetRank());
+		Util::FSmallUIntArray	OutputShape;
 
-		for (int32& Dim : Axes)
+		if (KeepDims)
 		{
-			checkf(Dim < ReducedDims.Num(), TEXT("Index out of bounds for Reduce axis"));
-			ReducedDims[Dim] = 1;
-			DmlAxes.Add(Dim);
+			OutputShape.Append(ReducedDims);
 		}
-
-		Util::FSmallUIntArray	OutputShape;
-
-		if (!KeepDims)
+		else
 		{
-			// ReduceSum example:
+			// Example:
 			//     input dims: {3, 2, 2}
 			//     axes: 1
 			//     keepDims: 0
 			// 
-			// ONNX: {3, 2}, but DML: {3, 1, 2}	
-			TConstArrayView<uint32> InputShape = InputTensors[0].GetShape().GetData();
+			// Shape: {3, 2}, but DML: {3, 1, 2}	
+			TConstArrayView<uint32> InputShape = InputTensors[0]->GetShape().GetData();
 
 			OutputShape.SetNum(ReducedDims.Num());
-
 			for (int32 Idx = 0; Idx < OutputShape.Num(); ++Idx)
 			{
-				OutputShape[Idx] = ReducedDims[Idx] ? 1 : InputShape[Idx];
+				OutputShape[Idx] = IsReducedDims[Idx] ? 1 : InputShape[Idx];
 			}
 		}
-		else
-		{
-			OutputShape.Append(OutputTensors[0].GetShape().GetData());
-		}
-
+		
 		FTensorDescDml DmlOutputTensorDesc;
 
 		if (!DmlOutputTensorDesc
-				.SetFromTensor(OutputTensors[0])
+				.SetFromTensor(*OutputTensors[0])
 				.SetShape(OutputShape)
 				.Validate())
 		{
@@ -138,22 +203,22 @@ public:
 		if constexpr (ReduceFunc == DML_REDUCE_FUNCTION_ARGMAX)
         {
             DML_ARGMAX_OPERATOR_DESC OpDesc;
-            OpDesc.AxisDirection = (DML_AXIS_DIRECTION) Attributes.GetValueOrDefault<int32>(TEXT("select_last_index"), 0);
+            OpDesc.AxisDirection = AxisDirection;
             OpDesc.InputTensor = DmlInputTensorDesc.GetDmlDesc();
 			OpDesc.OutputTensor = DmlOutputTensorDesc.GetDmlDesc();
-            OpDesc.Axes = DmlAxes.GetData();
-            OpDesc.AxisCount = (uint32) DmlAxes.Num();
+            OpDesc.Axes = Axes.GetData();
+            OpDesc.AxisCount = (uint32) Axes.Num();
 
 			return CreateOperator(Device, DML_OPERATOR_DESC { DML_OPERATOR_ARGMAX, &OpDesc });
         }
         else if constexpr (ReduceFunc == DML_REDUCE_FUNCTION_ARGMIN)
         {
             DML_ARGMIN_OPERATOR_DESC OpDesc;
-            OpDesc.AxisDirection = (DML_AXIS_DIRECTION) Attributes.GetValueOrDefault<int32>(TEXT("select_last_index"), 0);
+            OpDesc.AxisDirection = AxisDirection;
             OpDesc.InputTensor = DmlInputTensorDesc.GetDmlDesc();
 			OpDesc.OutputTensor = DmlOutputTensorDesc.GetDmlDesc();
-            OpDesc.Axes = DmlAxes.GetData();
-            OpDesc.AxisCount = (uint32) DmlAxes.Num();
+            OpDesc.Axes = Axes.GetData();
+            OpDesc.AxisCount = (uint32) Axes.Num();
 
             return CreateOperator(Device, DML_OPERATOR_DESC { DML_OPERATOR_ARGMIN, &OpDesc });
         }
@@ -164,8 +229,8 @@ public:
 			OpDesc.InputTensor = DmlInputTensorDesc.GetDmlDesc();
 			OpDesc.OutputTensor = DmlOutputTensorDesc.GetDmlDesc();
 			OpDesc.Function = ReduceFunc;
-			OpDesc.Axes = DmlAxes.GetData();
-			OpDesc.AxisCount = (uint32) DmlAxes.Num();
+			OpDesc.Axes = Axes.GetData();
+			OpDesc.AxisCount = (uint32) Axes.Num();
 
 			return CreateOperator(Device, DML_OPERATOR_DESC{ DML_OPERATOR_REDUCE, &OpDesc} );
 		}
diff --git a/Engine/Plugins/Experimental/NNERuntimeRDG/Source/NNERuntimeRDG/Private/Dml/NNEDmlOperatorResample.cpp b/Engine/Plugins/Experimental/NNERuntimeRDG/Source/NNERuntimeRDG/Private/Dml/NNEDmlOperatorResample.cpp
--- a/Engine/Plugins/Experimental/NNERuntimeRDG/Source/NNERuntimeRDG/Private/Dml/NNEDmlOperatorResample.cpp
+++ b/Engine/Plugins/Experimental/NNERuntimeRDG/Source/NNERuntimeRDG/Private/Dml/NNEDmlOperatorResample.cpp
@@ -32,6 +32,25 @@ void RemoveValuesByIndex(TConstArrayView<uint32> Indices, TArray<TData, TAllocat
 template <bool IsResize>
 class FOperatorDmlResample : public FOperatorDml
 {
+	//
+	//
+	//
+	enum ECoordTransformMode : uint8
+	{
+		None,
+		AlignCorners,
+		Asymmetric,
+		TfHalfPixelForNN,
+		TfCropAndResize,
+		HalfPixel
+	};
+
+	DML_INTERPOLATION_MODE	Mode;
+	ECoordTransformMode		CoordTransformMode;
+
+	//
+	//
+	//
 	static DML_INTERPOLATION_MODE ModeFromString(FStringView StringVal)
 	{
 		if (FCString::Stricmp(StringVal.GetData(), TEXT("NEAREST")) == 0)
@@ -51,11 +70,17 @@ class FOperatorDmlResample : public FOperatorDml
 
 public:
 
+	//
+	//
+	//
 	static FOperatorDml* Create()
 	{
 		return new FOperatorDmlResample<IsResize>();
 	}
 
+	//
+	//
+	//
 	static bool Validate(const NNE::FAttributeMap& AttributeMap, TConstArrayView<ENNETensorDataType> InputTypes, TConstArrayView<NNE::FSymbolicTensorShape> InputShapes)
 	{
 		//TODO
@@ -65,121 +90,196 @@ public:
 	//
 	//
 	//
-	virtual bool Initialize(IDMLDevice* Device, TArrayView<const NNE::Internal::FTensor> InputTensors, TArrayView<const NNE::Internal::FTensor> OutputTensors, const NNE::FAttributeMap& Attributes) override
+	virtual bool Initialize(TConstArrayView<NNE::FTensorDesc> Inputs, TConstArrayView<NNE::FTensorDesc> Outputs, const NNE::FAttributeMap& Attributes) override
 	{
-		check(InputTensors.Num() >= 1 && InputTensors.Num() < 4);
-		check(OutputTensors.Num() == 1);
+		check(Inputs.Num() >= 1 && Inputs.Num() < 4);
+		check(Outputs.Num() == 1);
 
-		const NNE::Internal::FTensor& InputTensor = InputTensors[0];
-		const NNE::Internal::FTensor& ScaleTensor = (InputTensors.Num() == 2)?InputTensors[1]:InputTensors[2];//Upsample have scale at position 1, while Resize at 2.
-		const NNE::Internal::FTensor& OutputTensor = OutputTensors[0];
-
-		if (InputTensors.Num() == 3 && InputTensors[2].GetDataType() == ENNETensorDataType::Int64)
+		if (Inputs.Num() == 3 && Inputs[2].GetDataType() == ENNETensorDataType::Int64)
 		{
-			UE_LOG(LogNNE, Warning, TEXT("Unsupported input type for 'sizes' of name %s, only 'scales' of type float is supported."), *InputTensors[2].GetName());
+			UE_LOG(LogNNE, Warning, TEXT("Unsupported input type for 'sizes' of name %s, only 'scales' of type float is supported."), *Inputs[2].GetName());
 			return false;
 		}
 
-		if (!ScaleTensor.HasPreparedData())
+		// Read attributes
+		Mode = ModeFromString(Attributes.GetValue<FString>(TEXT("mode")));
+
+		if constexpr (IsResize)
 		{
-			UE_LOG(LogNNE, Warning, TEXT("scales should be a constant tensor, it is here a variable tensor of name %s."), *ScaleTensor.GetName());
-			return false;
+			if (Mode == DML_INTERPOLATION_MODE_NEAREST_NEIGHBOR)
+			{
+				FString NeareastMode = Attributes.GetValueOrDefault<FString>(TEXT("nearest_mode"), TEXT("round_prefer_floor"));
+				if (FCString::Stricmp(*NeareastMode, TEXT("floor")) != 0)
+				{
+					UE_LOG(LogNNE, Warning, TEXT("Unsupported neareast mode:%s, using floor instead"), *NeareastMode);
+				}
+			}
+
+			FString CoordinateTransformationMode = Attributes.GetValueOrDefault<FString>(TEXT("coordinate_transformation_mode"), TEXT("half_pixel"));
+			
+			if (CoordinateTransformationMode == TEXT("align_corners"))
+			{
+				CoordTransformMode = ECoordTransformMode::AlignCorners;
+			}
+			else if (CoordinateTransformationMode == TEXT("asymmetric"))
+			{
+				CoordTransformMode = ECoordTransformMode::Asymmetric;
+			}
+			else if (CoordinateTransformationMode == TEXT("tf_half_pixel_for_nn"))
+			{
+				CoordTransformMode = ECoordTransformMode::TfHalfPixelForNN;
+			}
+			else if (CoordinateTransformationMode == TEXT("tf_crop_and_resize"))
+			{
+				CoordTransformMode = ECoordTransformMode::TfCropAndResize;
+			}
+			else
+			{
+				CoordTransformMode = ECoordTransformMode::HalfPixel;
+
+				if (CoordinateTransformationMode != TEXT("half_pixel"))
+				{
+					UE_LOG(LogNNE, Warning, TEXT("Unsupported coordinate transformation mode:%s, using half_pixel instead"), *CoordinateTransformationMode);
+				}
+			}
 		}
 
-		if (ScaleTensor.GetShape().Volume() != InputTensor.GetShape().Rank())
+		// Notify DML graph that we have constant CPU tensor inputs
+		for (int i = 1; i < Inputs.Num(); ++i)
 		{
-			UE_LOG(LogNNE, Warning, TEXT("scales tensor should contain N entries, where N is rank of X."));
-			return false;
+			ConstantCPUInputs.Add(i);
 		}
-		
-		// Read attributes
-		DML_INTERPOLATION_MODE Mode = ModeFromString(Attributes.GetValue<FString>(TEXT("mode")));
 
-		Util::FSmallArray<float> InputPixelOffsets, OutputPixelOffsets;
-		
-		InputPixelOffsets.Init(0.5f, InputTensor.GetShape().Rank());
-		OutputPixelOffsets.Init(-0.5f, InputTensor.GetShape().Rank());
+		return true;
+	}
+
+	//
+	//
+	//
+	virtual int PrepareOutputs(TConstArrayView<NNE::Internal::FTensorRef> InputTensors, TArrayView<NNE::Internal::FTensorRef> OutputTensors) const override
+	{
+		const NNE::Internal::FTensor& InputTensor = *InputTensors[0];
+		const NNE::Internal::FTensor& ScaleTensor = (InputTensors.Num() == 2) ? *InputTensors[1] : *InputTensors[2]; // Upsample has scale at position 1, while Resize at 2
+		const NNE::Internal::FTensor& OutputTensor = *OutputTensors[0];
+
+		if (!ScaleTensor.HasPreparedData())
+		{
+			UE_LOG(LogNNE, Error, TEXT("scales should be a constant tensor, it is here a variable tensor of name %s."), *ScaleTensor.GetName());
+			return -1;
+		}
 
-		Util::FSmallArray<float> Scales ( ScaleTensor.GetPreparedData<float>() );
+		if (ScaleTensor.GetShape().Volume() != InputTensor.GetShape().Rank())
+		{
+			UE_LOG(LogNNE, Error, TEXT("scales tensor should contain N entries, where N is rank of X."));
+			return -1;
+		}
 
 		if constexpr (IsResize)
 		{
-			if (Mode == DML_INTERPOLATION_MODE_NEAREST_NEIGHBOR)
+			if (CoordTransformMode == ECoordTransformMode::TfCropAndResize)
 			{
-				FString NeareastMode = Attributes.GetValueOrDefault<FString>(TEXT("nearest_mode"), TEXT("round_prefer_floor"));
-				if (FCString::Stricmp(*NeareastMode, TEXT("floor")) != 0)
+				const NNE::Internal::FTensor& RoiTensor = *InputTensors[1];
+
+				if (!RoiTensor.HasPreparedData())
 				{
-					UE_LOG(LogNNE, Warning, TEXT("Unsupported neareast mode:%s, using floor instead"), *NeareastMode);
+					UE_LOG(LogNNE, Warning, TEXT("roi should be a constant tensor, it is here a variable tensor of name %s."), *RoiTensor.GetName());
+					return -1;
+				}
+
+				if (RoiTensor.GetShape().Rank() != 1)
+				{
+					UE_LOG(LogNNE, Warning, TEXT("roi tensor should be 1-D."));
+					return -1;
+				}
+
+				if (RoiTensor.GetShape().GetData()[0] != 2 * InputTensor.GetShape().Rank())
+				{
+					UE_LOG(LogNNE, Warning, TEXT("roi tensor should contain 2*N entries, where N is rank of X."));
+					return -1;
+				}
+
+				for (int Idx = 0; Idx < InputTensor.GetShape().Rank(); ++Idx)
+				{
+					float LengthResized = (float)OutputTensor.GetShape().GetData()[Idx];
+
+					if (LengthResized <= 1)
+					{
+						UE_LOG(LogNNE, Warning, TEXT("Unsupported combination of transformation mode tf_crop_and_resize and length of resized dimension (%d) <= 1"), Idx);
+						return -1;
+					}
 				}
 			}
+		}
 
-			FString CoordinateTransformationMode = Attributes.GetValueOrDefault<FString>(TEXT("coordinate_transformation_mode"), TEXT("half_pixel"));
+		TConstArrayView<float>	ScalesData = ScaleTensor.GetPreparedData<float>();
+		TConstArrayView<uint32>	InputShape = InputTensor.GetShape().GetData();
+		TArray<uint32>			OutputShape;
+
+		for (int32 i = 0; i < InputShape.Num(); ++i)
+		{
+			OutputShape.Emplace(FMath::FloorToInt32(InputShape[i] * ScalesData[i]));
+		}
+
+		OutputTensors[0]->SetShape(NNE::FTensorShape::Make(OutputShape));
+
+		return 0;
+	}
+
+	//
+	//
+	//
+	virtual bool Create(IDMLDevice* Device, TConstArrayView<NNE::Internal::FTensorRef> InputTensors, TConstArrayView<NNE::Internal::FTensorRef> OutputTensors) override
+	{
+		const NNE::Internal::FTensor& InputTensor = *InputTensors[0];
+		const NNE::Internal::FTensor& ScaleTensor = (InputTensors.Num() == 2) ? *InputTensors[1] : *InputTensors[2]; // Upsample has scale at position 1, while Resize at 2
+		const NNE::Internal::FTensor& OutputTensor = *OutputTensors[0];
+
+		Util::FSmallArray<float> InputPixelOffsets, OutputPixelOffsets;
+
+		InputPixelOffsets.Init(0.5f, InputTensor.GetShape().Rank());
+		OutputPixelOffsets.Init(-0.5f, InputTensor.GetShape().Rank());
+
+		Util::FSmallArray<float> Scales(ScaleTensor.GetPreparedData<float>());
+
+		if constexpr (IsResize)
+		{
 			for (int Idx = 0; Idx < InputTensor.GetShape().Rank(); ++Idx)
 			{
 				float LengthResized = (float) OutputTensor.GetShape().GetData()[Idx];
 				float LengthOriginal = (float) InputTensor.GetShape().GetData()[Idx];
-				
-				if (CoordinateTransformationMode == TEXT("align_corners"))
+
+				if (CoordTransformMode == ECoordTransformMode::AlignCorners)
 				{
 					Scales[Idx] = (LengthResized - 1.0f) / (LengthOriginal - 1.0f);
 					InputPixelOffsets[Idx] = 0.f;
 					OutputPixelOffsets[Idx] = 0.f;
 				}
-				else if (CoordinateTransformationMode == TEXT("asymmetric"))
+				else if (CoordTransformMode == ECoordTransformMode::Asymmetric)
 				{
 					InputPixelOffsets[Idx] = 0.f;
 					OutputPixelOffsets[Idx] = 0.f;
 				}
-				else if (CoordinateTransformationMode == TEXT("tf_half_pixel_for_nn"))
+				else if (CoordTransformMode == ECoordTransformMode::TfHalfPixelForNN)
 				{
 					InputPixelOffsets[Idx] = 0.0f;
 					OutputPixelOffsets[Idx] = -0.5f;
 				}
-				else if (CoordinateTransformationMode == TEXT("tf_crop_and_resize"))
+				else if (CoordTransformMode == ECoordTransformMode::TfCropAndResize)
 				{
 					// NOTE: no tests for this, ORT erroneously puts all 0.0fs in the output tensor in this case.
-					
-					const NNE::Internal::FTensor& RoiTensor = InputTensors[1];
-					if (!RoiTensor.HasPreparedData())
-					{
-						UE_LOG(LogNNE, Warning, TEXT("roi should be a constant tensor, it is here a variable tensor of name %s."), *RoiTensor.GetName());
-						return false;
-					}
-					if (RoiTensor.GetShape().Rank() != 1)
-					{
-						UE_LOG(LogNNE, Warning, TEXT("roi tensor should be 1-D."));
-						return false;
-					}
-					if (RoiTensor.GetShape().GetData()[0] != 2 * InputTensor.GetShape().Rank())
-					{
-						UE_LOG(LogNNE, Warning, TEXT("roi tensor should contain 2*N entries, where N is rank of X."));
-						return false;
-					}
+					const NNE::Internal::FTensor& RoiTensor = *InputTensors[1];
 					
 					float Start = RoiTensor.GetPreparedData<float>()[Idx];
 					float End = RoiTensor.GetPreparedData<float>()[Idx + InputTensor.GetShape().Rank()];
+
 					if (LengthResized > 1)
 					{
-						Scales[Idx] = (LengthResized - 1.0f) / FMath::Max( (End - Start) * (LengthOriginal - 1.0f), 1.0f );
+						Scales[Idx] = (LengthResized - 1.0f) / FMath::Max((End - Start) * (LengthOriginal - 1.0f), 1.0f);
 						InputPixelOffsets[Idx] = Start * (1.0f - LengthOriginal);
 						OutputPixelOffsets[Idx] = 0.0f;
 					}
-					else
-					{
-						UE_LOG(LogNNE, Warning, TEXT("Unsupported combination of transformation mode tf_crop_and_resize and length of resized dimension (%d) <= 1"), Idx);
-						return false;
-					}
-					
-				}
-				else
-				{
-					if (CoordinateTransformationMode != TEXT("half_pixel"))
-					{
-						UE_LOG(LogNNE, Warning, TEXT("Unsupported coordinate transformation mode:%s, using half_pixel instead"), *CoordinateTransformationMode);
-					}
 				}
 			}
-			
 		}
 
 		// Initialize tensor descriptors
@@ -193,12 +293,6 @@ public:
 		DmlOutputTensorDesc
 			.SetTensorRank(1, 4)
 			.SetFromTensor(OutputTensor);
-
-		// Notify DML graph that we have constant CPU tensor inputs
-		for (int i = 1; i < InputTensors.Num(); ++i)
-		{
-			ConstantCPUInputs.Add(i);
-		}
 		
 		// Find any useless dimensions of size 1 that occur in both input and output
 		Util::FSmallUIntArray		SqueezeInds;
diff --git a/Engine/Plugins/Experimental/NNERuntimeRDG/Source/NNERuntimeRDG/Private/Dml/NNEDmlOperatorReshape.cpp b/Engine/Plugins/Experimental/NNERuntimeRDG/Source/NNERuntimeRDG/Private/Dml/NNEDmlOperatorReshape.cpp
--- a/Engine/Plugins/Experimental/NNERuntimeRDG/Source/NNERuntimeRDG/Private/Dml/NNEDmlOperatorReshape.cpp
+++ b/Engine/Plugins/Experimental/NNERuntimeRDG/Source/NNERuntimeRDG/Private/Dml/NNEDmlOperatorReshape.cpp
@@ -16,14 +16,22 @@ namespace UE::NNERuntimeRDG::Private::Dml
  */
 class FOperatorDmlReshape : public FOperatorDml
 {
+	mutable Util::FSmallUIntArray	OutputShape;
+	bool							bAllowZero;
 
 public:
 
+	//
+	//
+	//
 	static FOperatorDml* Create()
 	{
 		return new FOperatorDmlReshape();
 	}
 
+	//
+	//
+	//
     static bool Validate(const NNE::FAttributeMap& AttributeMap, TConstArrayView<ENNETensorDataType> InputTypes, TConstArrayView<NNE::FSymbolicTensorShape> InputShapes)
 	{
 		//TODO
@@ -33,53 +41,71 @@ public:
 	//
 	//
 	//
-	virtual bool Initialize(IDMLDevice* Device, TArrayView<const NNE::Internal::FTensor> InputTensors, TArrayView<const NNE::Internal::FTensor> OutputTensors, const NNE::FAttributeMap& Attributes) override
+	virtual bool Initialize(TConstArrayView<NNE::FTensorDesc> Inputs, TConstArrayView<NNE::FTensorDesc> Outputs, const NNE::FAttributeMap& Attributes) override
 	{
-        check(InputTensors.Num() == 2);
-        check(OutputTensors.Num() == 1);
-
-        ConstantCPUInputs.Add(1);
-
-        // Shape tensor must be constant!
-        check(InputTensors[1].HasPreparedData());
-
-        bool bAllowZero = (bool)
-			( Attributes.GetValueOrDefault<int32>(TEXT("allowzero"), 0) );
-
-        Util::FSmallUIntArray ReshapedShape;
-
-        switch(InputTensors[1].GetDataType())
-        {
-        case ENNETensorDataType::Int32:
-            if(!ShapeHelper::Reshape::ReshapeTensor<int32>(InputTensors[0].GetShape(), InputTensors[1], bAllowZero, ReshapedShape))
-            {
-                return false;
-            }
-            break;
-        case ENNETensorDataType::Int64:
-			if (!ShapeHelper::Reshape::ReshapeTensor<int64>(InputTensors[0].GetShape(), InputTensors[1], bAllowZero, ReshapedShape))
-            {
-                return false;
-            }
-            break;
-        case ENNETensorDataType::UInt32:
-			if (!ShapeHelper::Reshape::ReshapeTensor<uint32>(InputTensors[0].GetShape(), InputTensors[1], bAllowZero, ReshapedShape))
-            {
-                return false;
-            }
-            break;
-        default:
-            UE_LOG(LogNNE, Warning, TEXT("Shape tensor has invalid data type"));
-			return false;
-        };
-        
-        check(ReshapedShape == OutputTensors[0].GetShape().GetData());
+		check(Inputs.Num() == 2);
+		check(Outputs.Num() == 1);
+
+		ConstantCPUInputs.Add(1);
+
+		bAllowZero = Attributes.GetValueOrDefault<int32>(TEXT("allowzero"), 0) != 0;
+
+		return true;
+	}
 
+	//
+	//
+	//
+	virtual int PrepareOutputs(TConstArrayView<NNE::Internal::FTensorRef> InputTensors, TArrayView<NNE::Internal::FTensorRef> OutputTensors) const override
+	{
+		// Shape tensor must be constant!
+		check(InputTensors[1]->HasPreparedData());
+		
+		OutputShape.Reset();
+
+		switch (InputTensors[1]->GetDataType())
+		{
+			case ENNETensorDataType::Int32:
+				if (!ShapeHelper::Reshape::ReshapeTensor<int32>(InputTensors[0]->GetShape(), *InputTensors[1], bAllowZero, OutputShape))
+				{
+					return false;
+				}
+				break;
+
+			case ENNETensorDataType::Int64:
+				if (!ShapeHelper::Reshape::ReshapeTensor<int64>(InputTensors[0]->GetShape(), *InputTensors[1], bAllowZero, OutputShape))
+				{
+					return false;
+				}
+				break;
+
+			case ENNETensorDataType::UInt32:
+				if (!ShapeHelper::Reshape::ReshapeTensor<uint32>(InputTensors[0]->GetShape(), *InputTensors[1], bAllowZero, OutputShape))
+				{
+					return false;
+				}
+				break;
+
+			default:
+				UE_LOG(LogNNE, Warning, TEXT("Shape tensor has invalid data type"));
+				return false;
+		};
+
+		OutputTensors[0]->SetShape(NNE::FTensorShape::Make(OutputShape));
+
+		return 0;
+	}
+
+	//
+	//
+	//
+	virtual bool Create(IDMLDevice* Device, TConstArrayView<NNE::Internal::FTensorRef> InputTensors, TConstArrayView<NNE::Internal::FTensorRef> OutputTensors) override
+	{
         FTensorDescDml DmlInputTensorDesc;
 
         if (!DmlInputTensorDesc
-				.SetFromTensor(InputTensors[0])
-				.SetShape(ReshapedShape)
+				.SetFromTensor(*InputTensors[0])
+				.SetShape(OutputShape)
 				.Validate())
 		{
 			UE_LOG(LogNNE, Warning, TEXT("Failed to initialize Reshape's input tensor for DML inference"));
@@ -89,7 +115,7 @@ public:
         FTensorDescDml DmlOutputTensorDesc;
 
         if (!DmlOutputTensorDesc
-				.SetFromTensor(OutputTensors[0])
+				.SetFromTensor(*OutputTensors[0])
 				.Validate())
 		{
 			UE_LOG(LogNNE, Warning, TEXT("Failed to initialize Reshape's output tensor for DML inference"));
diff --git a/Engine/Plugins/Experimental/NNERuntimeRDG/Source/NNERuntimeRDG/Private/Dml/NNEDmlOperatorSlice.cpp b/Engine/Plugins/Experimental/NNERuntimeRDG/Source/NNERuntimeRDG/Private/Dml/NNEDmlOperatorSlice.cpp
--- a/Engine/Plugins/Experimental/NNERuntimeRDG/Source/NNERuntimeRDG/Private/Dml/NNEDmlOperatorSlice.cpp
+++ b/Engine/Plugins/Experimental/NNERuntimeRDG/Source/NNERuntimeRDG/Private/Dml/NNEDmlOperatorSlice.cpp
@@ -17,15 +17,19 @@ namespace UE::NNERuntimeRDG::Private::Dml
  */
 class FOperatorDmlSlice : public FOperatorDml
 {	
+	//
+	//
+	//
 	template<typename DataType>
-	void ComputeOffsetsSizesStrides(
-		TArrayView<const NNE::Internal::FTensor> InputTensors,
+	static void ComputeOffsetsSizesStrides(
+		TArrayView<const NNE::Internal::FTensorRef> InputTensors,
 		Util::FSmallUIntArray& OutOffsets, 
 		Util::FSmallUIntArray& OutSizes, 
-		Util::FSmallIntArray& OutStrides)
+		Util::FSmallIntArray& OutStrides,
+		Util::FSmallUIntArray& OutputShape)
 	{
-		TConstArrayView<DataType> Starts =  InputTensors[1].GetPreparedData<DataType>();
-		TConstArrayView<DataType> Ends = InputTensors[2].GetPreparedData<DataType>();
+		TConstArrayView<DataType> Starts =  InputTensors[1]->GetPreparedData<DataType>();
+		TConstArrayView<DataType> Ends = InputTensors[2]->GetPreparedData<DataType>();
 		check(Starts.Num() == Ends.Num());
 
 		Util::FSmallArray<DataType> Axes;
@@ -33,14 +37,14 @@ class FOperatorDmlSlice : public FOperatorDml
 
 		if (InputTensors.Num() >= 4)
 		{
-			auto NormalizeAxes = [NumDims = (DataType) InputTensors[0].GetShape().Rank()] (DataType& Axis)
+			auto NormalizeAxes = [NumDims = (DataType) InputTensors[0]->GetShape().Rank()] (DataType& Axis)
 			{
 				if(Axis < 0)
 				{
 					Axis += NumDims;
 				}
 			};
-			Axes = InputTensors[3].GetPreparedData<DataType>();
+			Axes = InputTensors[3]->GetPreparedData<DataType>();
 			check(Axes.Num() == Starts.Num());
 			Algo::ForEach(Axes, NormalizeAxes);
 		}
@@ -52,16 +56,16 @@ class FOperatorDmlSlice : public FOperatorDml
 
 		if (InputTensors.Num() >= 5)
 		{
-			Steps = InputTensors[4].GetPreparedData<DataType>();
+			Steps = InputTensors[4]->GetPreparedData<DataType>();
 			check(Steps.Num() == Axes.Num());
 		}
 		else
 		{
-			Steps.Init((DataType)1, Axes.Num());
+			Steps.Init((DataType) 1, Axes.Num());
 		}
 
-		Util::FSmallUIntArray OutputShape;
-		Algo::Copy(InputTensors[0].GetShape().GetData(), OutputShape);
+		OutputShape.Reset();
+		Algo::Copy(InputTensors[0]->GetShape().GetData(), OutputShape);
 		OutSizes = OutputShape;
 		OutOffsets.SetNumZeroed(OutputShape.Num());
 
@@ -74,11 +78,11 @@ class FOperatorDmlSlice : public FOperatorDml
 			DataType End = Ends[Idx];
 			
 			DataType DimIndex = Axes[Idx];
-			check(DimIndex < (DataType) InputTensors[0].GetShape().Rank());
+			check(DimIndex < (DataType) InputTensors[0]->GetShape().Rank());
 			DataType Stride = Steps[Idx];
 			check(Stride != 0);
 
-			uint32 Dim = InputTensors[0].GetShape().GetData()[DimIndex];
+			uint32 Dim = InputTensors[0]->GetShape().GetData()[DimIndex];
 			if (Start < 0 && Start > TNumericLimits<DataType>::Min())
 			{
 				Start += (DataType) Dim;
@@ -107,13 +111,25 @@ class FOperatorDmlSlice : public FOperatorDml
 		}
 	}
 
+	mutable Util::FSmallUIntArray	OutputShape;
+	mutable Util::FSmallUIntArray	Offsets, Sizes;
+	mutable Util::FSmallIntArray	Strides;
+
+	ENNETensorDataType InputIndexDataType;
+
 public:
 
+	//
+	//
+	//
 	static FOperatorDml* Create()
 	{
 		return new FOperatorDmlSlice();
 	}
 
+	//
+	//
+	//
 	static bool Validate(const NNE::FAttributeMap& AttributeMap, TConstArrayView<ENNETensorDataType> InputTypes, TConstArrayView<NNE::FSymbolicTensorShape> InputShapes)
 	{
 		//TODO
@@ -123,36 +139,73 @@ public:
 	//
 	//
 	//
-	virtual bool Initialize(IDMLDevice* Device, TArrayView<const NNE::Internal::FTensor> InputTensors, TArrayView<const NNE::Internal::FTensor> OutputTensors, const NNE::FAttributeMap& Attributes) override
+	virtual bool Initialize(TConstArrayView<NNE::FTensorDesc> Inputs, TConstArrayView<NNE::FTensorDesc> Outputs, const NNE::FAttributeMap& Attributes) override
 	{
-		check(InputTensors.Num() >= 3);
-		check(InputTensors.Num() <= 5);
-		check(OutputTensors.Num() == 1);
+		check(Inputs.Num() >= 3);
+		check(Inputs.Num() <= 5);
+		check(Outputs.Num() == 1);
 
-		ENNETensorDataType InputIndexDataType = ENNETensorDataType::None;
+		InputIndexDataType = ENNETensorDataType::None;
 
-		for(int Idx = 1; Idx < InputTensors.Num(); Idx++)
+		for (int Idx = 1; Idx < Inputs.Num(); Idx++)
 		{
-			check(InputTensors[Idx].GetShape().Rank() == 1);
+			check(Inputs[Idx].GetShape().Rank() == 1);
+			
 			if (Idx == 1)
 			{
-				InputIndexDataType = InputTensors[Idx].GetDataType();
-				check( InputIndexDataType == ENNETensorDataType::Int32 
-			   		|| InputIndexDataType == ENNETensorDataType::Int64);
+				InputIndexDataType = Inputs[Idx].GetDataType();
+				if (InputIndexDataType != ENNETensorDataType::Int32 && InputIndexDataType != ENNETensorDataType::Int64)
+				{
+					UE_LOG(LogNNE, Error, TEXT("Failed to initialize Slice input tensor at index 1, data type needs to be Int32 or Int64"));
+					return false;
+				}
 			}
 			else
 			{
-				check(InputTensors[Idx].GetDataType() == InputIndexDataType);
+				if (Inputs[Idx].GetDataType() != InputIndexDataType)
+				{
+					UE_LOG(LogNNE, Error, TEXT("Failed to initialize Slice input tensor at index %d, data type needs to same as at index 1"), Idx);
+					return false;
+				}
 			}
-			check(InputTensors[Idx].HasPreparedData());
+
 			ConstantCPUInputs.Add(Idx);
 		}
 
+		return true;
+	}
+
+	//
+	//
+	//
+	virtual int PrepareOutputs(TConstArrayView<NNE::Internal::FTensorRef> InputTensors, TArrayView<NNE::Internal::FTensorRef> OutputTensors) const override
+	{
+		switch (InputIndexDataType)
+		{
+			case ENNETensorDataType::Int32:
+				ComputeOffsetsSizesStrides<int32>(InputTensors, Offsets, Sizes, Strides, OutputShape);
+				break;
+
+			case ENNETensorDataType::Int64:
+				ComputeOffsetsSizesStrides<int64>(InputTensors, Offsets, Sizes, Strides, OutputShape);
+				break;
+		}
+
+		OutputTensors[0]->SetShape(NNE::FTensorShape::Make(OutputShape));
+
+		return 0;
+	}
+
+	//
+	//
+	//
+	virtual bool Create(IDMLDevice* Device, TConstArrayView<NNE::Internal::FTensorRef> InputTensors, TConstArrayView<NNE::Internal::FTensorRef> OutputTensors) override
+	{
 		// Initialize Input tensor desc
         FTensorDescDml DmlInputTensorDesc;
         
 		if (!DmlInputTensorDesc
-				.SetFromTensor(InputTensors[0])
+				.SetFromTensor(*InputTensors[0])
 				.Validate())
         {
             UE_LOG(LogNNE, Error, TEXT("Failed to initialize Slice input for DML inference"));
@@ -163,26 +216,13 @@ public:
         FTensorDescDml DmlOutputTensorDesc;
         
 		if (!DmlOutputTensorDesc
-				.SetFromTensor(OutputTensors[0])
+				.SetFromTensor(*OutputTensors[0])
 				.Validate())
         {
             UE_LOG(LogNNE, Error, TEXT("Failed to initialize Slice Output for DML inference"));
             return false;
         }
 
-		Util::FSmallUIntArray Offsets, Sizes;
-		Util::FSmallIntArray Strides;
-
-		switch(InputIndexDataType)
-		{
-		case ENNETensorDataType::Int32:
-			ComputeOffsetsSizesStrides<int32>(InputTensors, Offsets, Sizes, Strides);
-			break;
-		case ENNETensorDataType::Int64:
-			ComputeOffsetsSizesStrides<int64>(InputTensors, Offsets, Sizes, Strides);
-			break;
-		}
-
 		DML_SLICE1_OPERATOR_DESC DmlSliceOpDesc{};
 		
 		DmlSliceOpDesc.InputTensor = DmlInputTensorDesc.GetDmlDesc();
diff --git a/Engine/Plugins/Experimental/NNERuntimeRDG/Source/NNERuntimeRDG/Private/Dml/NNEDmlOperatorSpaceToDepth.cpp b/Engine/Plugins/Experimental/NNERuntimeRDG/Source/NNERuntimeRDG/Private/Dml/NNEDmlOperatorSpaceToDepth.cpp
--- a/Engine/Plugins/Experimental/NNERuntimeRDG/Source/NNERuntimeRDG/Private/Dml/NNEDmlOperatorSpaceToDepth.cpp
+++ b/Engine/Plugins/Experimental/NNERuntimeRDG/Source/NNERuntimeRDG/Private/Dml/NNEDmlOperatorSpaceToDepth.cpp
@@ -8,17 +8,32 @@
 namespace UE::NNERuntimeRDG::Private::Dml
 {
 
+//
+//
+//
 class FOperatorDmlSpaceToDepth : public FOperatorDml
 {
+	enum InputDims
+	{
+		N, C, H, W,
+		DIM_COUNT
+	};
 
+	int32 BlockSize;
 
 public:
 
+	//
+	//
+	//
 	static FOperatorDml* Create()
 	{
 		return new FOperatorDmlSpaceToDepth();
 	}
 
+	//
+	//
+	//
 	static bool Validate(const NNE::FAttributeMap& AttributeMap, TConstArrayView<ENNETensorDataType> InputTypes, TConstArrayView<NNE::FSymbolicTensorShape> InputShapes)
 	{
 		//TODO
@@ -28,12 +43,10 @@ public:
 	//
 	//
 	//
-	virtual bool Initialize(IDMLDevice* Device, TArrayView<const NNE::Internal::FTensor> InputTensors, TArrayView<const NNE::Internal::FTensor> OutputTensors, const NNE::FAttributeMap& Attributes) override
+	virtual bool Initialize(TConstArrayView<NNE::FTensorDesc> Inputs, TConstArrayView<NNE::FTensorDesc> Outputs, const NNE::FAttributeMap& Attributes) override
 	{
-		check(InputTensors.Num() == 1);
-		check(OutputTensors.Num() == 1);
-
-		int32 BlockSize;
+		check(Inputs.Num() == 1);
+		check(Outputs.Num() == 1);
 
 		const FNNEAttributeValue* BlockSizeAttr = Attributes.GetAttributeValue(TEXT("blocksize"));
 		if (BlockSizeAttr)
@@ -46,9 +59,35 @@ public:
 			return false;
 		}
 
+		return true;
+	}
 
-		const NNE::Internal::FTensor& InputTensor = InputTensors[0];
-		const NNE::Internal::FTensor& OutputTensor = OutputTensors[0];
+	//
+	//
+	//
+	virtual int PrepareOutputs(TConstArrayView<NNE::Internal::FTensorRef> InputTensors, TArrayView<NNE::Internal::FTensorRef> OutputTensors) const override
+	{
+		TConstArrayView<uint32>		InputShape = InputTensors[0]->GetShape().GetData();
+		Util::FSmallUIntArray		OutputShape;
+
+		OutputShape.SetNum(DIM_COUNT);
+		OutputShape[N] = InputShape[N];
+		OutputShape[C] = InputShape[C] * BlockSize * BlockSize;
+		OutputShape[H] = InputShape[H] / BlockSize;
+		OutputShape[W] = InputShape[W] / BlockSize;
+
+		OutputTensors[0]->SetShape(NNE::FTensorShape::Make(OutputShape));
+
+		return 0;
+	}
+
+	//
+	//
+	//
+	virtual bool Create(IDMLDevice* Device, TConstArrayView<NNE::Internal::FTensorRef> InputTensors, TConstArrayView<NNE::Internal::FTensorRef> OutputTensors) override
+	{
+		const NNE::Internal::FTensor& InputTensor = *InputTensors[0];
+		const NNE::Internal::FTensor& OutputTensor = *OutputTensors[0];
 
 		FTensorDescDml	DmlInputTensorDesc;
 		FTensorDescDml	DmlOutputTensorDesc;
diff --git a/Engine/Plugins/Experimental/NNERuntimeRDG/Source/NNERuntimeRDG/Private/Dml/NNEDmlOperatorSplit.cpp b/Engine/Plugins/Experimental/NNERuntimeRDG/Source/NNERuntimeRDG/Private/Dml/NNEDmlOperatorSplit.cpp
--- a/Engine/Plugins/Experimental/NNERuntimeRDG/Source/NNERuntimeRDG/Private/Dml/NNEDmlOperatorSplit.cpp
+++ b/Engine/Plugins/Experimental/NNERuntimeRDG/Source/NNERuntimeRDG/Private/Dml/NNEDmlOperatorSplit.cpp
@@ -2,21 +2,31 @@
 
 #ifdef NNE_USE_DIRECTML
 #include "NNEDmlOperator.h"
+#include "NNEDmlOperatorUtils.h"
 
 namespace UE::NNERuntimeRDG::Private::Dml
 {
-
+//
+//
+//
 class FOperatorDmlSplit : public FOperatorDml
 {
-
+	TArray<int32>	Split;
+	int				Axis;
 
 public:
 
+	//
+	//
+	//
 	static FOperatorDml* Create()
 	{
 		return new FOperatorDmlSplit();
 	}
 
+	//
+	//
+	//
 	static bool Validate(const NNE::FAttributeMap& AttributeMap, TConstArrayView<ENNETensorDataType> InputTypes, TConstArrayView<NNE::FSymbolicTensorShape> InputShapes)
 	{
 		//TODO
@@ -26,56 +36,106 @@ public:
 	//
 	//
 	//
-	virtual bool Initialize(IDMLDevice* Device, TArrayView<const NNE::Internal::FTensor> InputTensors, TArrayView<const NNE::Internal::FTensor> OutputTensors, const NNE::FAttributeMap& Attributes) override
+	virtual bool Initialize(TConstArrayView<NNE::FTensorDesc> Inputs, TConstArrayView<NNE::FTensorDesc> Outputs, const NNE::FAttributeMap& Attributes) override
 	{
-		check(InputTensors.Num() == 1);
-		const NNE::Internal::FTensor& InputTensor = InputTensors[0];
+		check(Inputs.Num() == 1);
+		
+		Axis = Attributes.GetValueOrDefault<int>(TEXT("axis"), 0);
+		Axis = HandleNegativeAxis(Axis, Inputs[0].GetShape().Rank());
 
-		int Axis = Attributes.GetValueOrDefault<int>(TEXT("axis"), 0);
+		const FNNEAttributeValue* AttrSplit = Attributes.GetAttributeValue(TEXT("split"));
 
-		// Check split size is correct
-		//TODO: code should be validated by a validator which should report to the user if the contract is broken. We do it here instead.
+		if (AttrSplit)
+		{
+			Split = AttrSplit->GetValue<TArray<int32>>();
+			
+			if (Split.Num() != Outputs.Num())
+			{
+				UE_LOG(LogNNE, Error, TEXT("Attribute split needs to have same count as number of outputs"));
+				return false;
+			}
 
-		uint32 SplitSize = 0;
+			if ((Split.Num() % Outputs.Num()) != 0)
+			{
+				UE_LOG(LogNNE, Error, TEXT("Attribute split count needs to be divisible by the number of outputs"));
+				return false;
+			}
+		}
 
-		for (int Idx = 0; Idx < OutputTensors.Num(); ++Idx)
+		// Check split size is correct
+		//TODO: code should be validated by a validator which should report to the user if the contract is broken. We do it here instead.
+		for (int Idx = 0; Idx < Outputs.Num(); ++Idx)
 		{
-			//check(OutputTensors[Idx].GetShape().Rank() == InputTensor.GetShape().Rank());
-			if (OutputTensors[Idx].GetShape().Rank() != InputTensor.GetShape().Rank())
+			if (Outputs[Idx].GetShape().Rank() != Inputs[0].GetShape().Rank())
 			{
 				UE_LOG(LogNNE, Error, TEXT("Rank of output tensor and input tensor should be the same"));
 				return false;
+			}			
+		}
+
+		return true;
+	}
+
+	//
+	//
+	//
+	virtual int PrepareOutputs(TConstArrayView<NNE::Internal::FTensorRef> InputTensors, TArrayView<NNE::Internal::FTensorRef> OutputTensors) const override
+	{
+		const NNE::Internal::FTensor&	InputTensor = *InputTensors[0];
+		TConstArrayView<uint32>			InputShape = InputTensor.GetShape().GetData();
+
+		if (!Split.IsEmpty())
+		{
+			int32 TotalElemCount = 0;
+
+			for (int32 ElemCount : Split)
+			{
+				TotalElemCount += ElemCount;
 			}
 
-			for (int Dim = 0; Dim < InputTensor.GetShape().Rank(); ++Dim)
+			if (TotalElemCount != InputShape[Axis])
 			{
-				if (Dim == Axis)
-				{
-					SplitSize += OutputTensors[Idx].GetShape().GetData()[Dim];
-				}
-				else
-				{
-					//check(OutputTensors[Idx].GetShape().GetData()[Dim] == InputTensor.GetShape().GetData()[Dim]);
-					if (OutputTensors[Idx].GetShape().GetData()[Dim] != InputTensor.GetShape().GetData()[Dim])
-					{
-						UE_LOG(LogNNE, Error, TEXT("%s"), *FString::Printf(TEXT("Output tensor %d 's dimension %d should match input tensor's"), Idx, Dim));
-						return false;
-					}
-				}
+				UE_LOG(LogNNE, Error, TEXT("Incorrect elem count for split axis"));
+				return -1;
 			}
-		}
+		
+			for (int32 Idx = 0; Idx < OutputTensors.Num(); ++Idx)
+			{
+				Util::FSmallUIntArray OutputShape;
+
+				OutputShape.Append(InputShape.GetData(), InputShape.Num());
+				OutputShape[Axis] = Split[Idx];
 
-		//check(SplitSize == InputTensor.GetShape().GetData()[Axis]);
-		if (SplitSize != InputTensor.GetShape().GetData()[Axis])
+				OutputTensors[Idx]->SetShape(NNE::FTensorShape::Make(OutputShape));
+			}
+		}
+		else
 		{
-			UE_LOG(LogNNE, Error, TEXT("Input tensor's axis dimension size must be equal to the sum of output tensors' axis dimensions sizes"));
-			return false;
+			uint32 EqualSplit = InputShape[Axis] / OutputTensors.Num();
+
+			for (int Idx = 0; Idx < OutputTensors.Num(); ++Idx)
+			{
+				Util::FSmallUIntArray OutputShape;
+
+				OutputShape.Append(InputShape.GetData(), InputShape.Num());
+				OutputShape[Axis] = EqualSplit;
+
+				OutputTensors[Idx]->SetShape(NNE::FTensorShape::Make(OutputShape));
+			}
 		}
 
+		return 0;
+	}
+
+	//
+	//
+	//
+	virtual bool Create(IDMLDevice* Device, TConstArrayView<NNE::Internal::FTensorRef> InputTensors, TConstArrayView<NNE::Internal::FTensorRef> OutputTensors) override
+	{
 		FTensorDescDml InputTensorDesc;
 		
 		if (!InputTensorDesc
-				.SetFromTensor(InputTensor)
+				.SetFromTensor(*InputTensors[0])
 				.Validate())
 		{
 			UE_LOG(LogNNE, Error, TEXT("Failed to initialize pooling operator's input tensor for DML inference"));
@@ -91,7 +151,7 @@ public:
 		for(int Idx = 0; Idx < OutputTensors.Num(); ++Idx)
 		{
 			if (!OutputTensorDescs[Idx]
-					.SetFromTensor(OutputTensors[Idx])
+					.SetFromTensor(*OutputTensors[Idx])
 					.Validate())
 			{
 				UE_LOG(LogNNE, Error, TEXT("Failed to initialize pooling operator's output tensor for DML inference"));
diff --git a/Engine/Plugins/Experimental/NNERuntimeRDG/Source/NNERuntimeRDG/Private/Dml/NNEDmlOperatorSqueeze.cpp b/Engine/Plugins/Experimental/NNERuntimeRDG/Source/NNERuntimeRDG/Private/Dml/NNEDmlOperatorSqueeze.cpp
--- a/Engine/Plugins/Experimental/NNERuntimeRDG/Source/NNERuntimeRDG/Private/Dml/NNEDmlOperatorSqueeze.cpp
+++ b/Engine/Plugins/Experimental/NNERuntimeRDG/Source/NNERuntimeRDG/Private/Dml/NNEDmlOperatorSqueeze.cpp
@@ -8,16 +8,26 @@
 namespace UE::NNERuntimeRDG::Private::Dml
 {
 
+//
+//
+//
 class FOperatorDmlSqueeze : public FOperatorDml
 {
+	mutable TArray<int32>	Axes;
 
 public:
 
+	//
+	//
+	//
 	static FOperatorDml* Create()
 	{
 		return new FOperatorDmlSqueeze();
 	}
 
+	//
+	//
+	//
 	static bool Validate(const NNE::FAttributeMap& AttributeMap, TConstArrayView<ENNETensorDataType> InputTypes, TConstArrayView<NNE::FSymbolicTensorShape> InputShapes)
 	{
 		//TODO
@@ -27,21 +37,30 @@ public:
 	//
 	//
 	//
-	virtual bool Initialize(IDMLDevice* Device, TArrayView<const NNE::Internal::FTensor> InputTensors, TArrayView<const NNE::Internal::FTensor> OutputTensors, const NNE::FAttributeMap& Attributes) override
+	virtual bool Initialize(TConstArrayView<NNE::FTensorDesc> Inputs, TConstArrayView<NNE::FTensorDesc> Outputs, const NNE::FAttributeMap& Attributes) override
 	{
-		check(InputTensors.Num() == 1);
-		check(OutputTensors.Num() == 1);
+		check(Inputs.Num() == 1);
+		check(Outputs.Num() == 1);
 
-		TConstArrayView<uint32>		InputShape = InputTensors[0].GetShape().GetData();
-		TArray<int32>				Axes;
+		const int32					InputShapeRank = Inputs[0].GetShape().Rank();
 		const FNNEAttributeValue*	AxesAttr = Attributes.GetAttributeValue(TEXT("axes"));
-		
+
 		if (AxesAttr)
 		{
 			Axes = AxesAttr->GetValue<TArray<int32>>();
-			HandleNegativeAxes(Axes, InputShape.Num());
+			HandleNegativeAxes(Axes, InputShapeRank);
 			Algo::Sort(Axes);
 		}
+
+		return true;
+	}
+
+	//
+	//
+	//
+	virtual int PrepareOutputs(TConstArrayView<NNE::Internal::FTensorRef> InputTensors, TArrayView<NNE::Internal::FTensorRef> OutputTensors) const override
+	{
+		TConstArrayView<uint32>	InputShape = InputTensors[0]->GetShape().GetData();
 		
 		if (Axes.IsEmpty())
 		{
@@ -58,25 +77,41 @@ public:
 
 		OutputShape.Reserve(InputShape.Num());
 		OutputShape.Append(InputShape);
-		
+
 		for (int32 Idx = Axes.Num() - 1; Idx >= 0; --Idx)
 		{
 			const int32 Axe = Axes[Idx];
 
 			if (OutputShape[Axe] != 1)
 			{
-				UE_LOG(LogNNE, Warning, TEXT("Squeeze at axe %d for 'Data' (name: %s) should be targeting a dimension of size 1 but it is %d."), Axe, *InputTensors[0].GetName(), OutputShape[Axe]);
+				UE_LOG(
+					LogNNE, 
+					Warning, 
+					TEXT("Squeeze at axe %d for 'Data' (name: %s) should be targeting a dimension of size 1 but it is %d."), 
+					Axe, 
+					*InputTensors[0]->GetName(), 
+					OutputShape[Axe]);
+				
 				return false;
 			}
 
 			OutputShape.RemoveAt(Axe);
 		}
 
+		OutputTensors[0]->SetShape(NNE::FTensorShape::Make(OutputShape));
+
+		return 0;
+	}
+
+	//
+	//
+	//
+	virtual bool Create(IDMLDevice* Device, TConstArrayView<NNE::Internal::FTensorRef> InputTensors, TConstArrayView<NNE::Internal::FTensorRef> OutputTensors) override
+	{
 		FTensorDescDml DmlTensorDesc;
 		
 		if (!DmlTensorDesc
-				.SetFromTensor(OutputTensors[0])
-				.SetShape(OutputShape)
+				.SetFromTensor(*OutputTensors[0])
 				.Validate())
 		{
 			UE_LOG(LogNNE, Warning, TEXT("Failed to initialize Unsqueeze's output tensor for DML inference"));
diff --git a/Engine/Plugins/Experimental/NNERuntimeRDG/Source/NNERuntimeRDG/Private/Dml/NNEDmlOperatorTranspose.cpp b/Engine/Plugins/Experimental/NNERuntimeRDG/Source/NNERuntimeRDG/Private/Dml/NNEDmlOperatorTranspose.cpp
--- a/Engine/Plugins/Experimental/NNERuntimeRDG/Source/NNERuntimeRDG/Private/Dml/NNEDmlOperatorTranspose.cpp
+++ b/Engine/Plugins/Experimental/NNERuntimeRDG/Source/NNERuntimeRDG/Private/Dml/NNEDmlOperatorTranspose.cpp
@@ -15,14 +15,36 @@ namespace UE::NNERuntimeRDG::Private::Dml
  */
 class FOperatorDmlTranspose : public FOperatorDml
 {
+	//
+	// Apply permutations to input array view
+	//
+	Util::FSmallUIntArray Permute(TConstArrayView<uint32> InputView) const
+	{
+		Util::FSmallUIntArray Permuted;
+
+		for (int32 PermVal : Perm)
+		{
+			Permuted.Add(InputView[PermVal]);
+		}
+
+		return Permuted;
+	};
+
+	TArray<int32> Perm;
 
 public:
 
+	//
+	//
+	//
 	static FOperatorDml* Create()
 	{
 		return new FOperatorDmlTranspose();
 	}
 
+	//
+	//
+	//
 	static bool Validate(const NNE::FAttributeMap& AttributeMap, TConstArrayView<ENNETensorDataType> InputTypes, TConstArrayView<NNE::FSymbolicTensorShape> InputShapes)
 	{
 		//TODO
@@ -32,49 +54,57 @@ public:
 	//
 	//
 	//
-	virtual bool Initialize(IDMLDevice* Device, TArrayView<const NNE::Internal::FTensor> InputTensors, TArrayView<const NNE::Internal::FTensor> OutputTensors, const NNE::FAttributeMap& Attributes) override
+	virtual bool Initialize(TConstArrayView<NNE::FTensorDesc> Inputs, TConstArrayView<NNE::FTensorDesc> Outputs, const NNE::FAttributeMap& Attributes) override
 	{
-		check(InputTensors.Num() == 1);
-		check(OutputTensors.Num() == 1);
+		check(Inputs.Num() == 1);
+		check(Outputs.Num() == 1);
 
-		int32 NumDims = InputTensors[0].GetShape().Rank();
+		const int32 NumDims = Inputs[0].GetShape().Rank();
 		check(NumDims > 0);
-		check(NumDims == OutputTensors[0].GetShape().Rank());
+		check(NumDims == Outputs[0].GetShape().Rank());
 
 		// Default permutation is reverse
 		Util::FSmallIntArray ReversePerm;
 
-		for(int Idx = NumDims-1; Idx >= 0; Idx--)
+		for (int Idx = NumDims - 1; Idx >= 0; Idx--)
 		{
 			ReversePerm.Add(Idx);
 		}
 
-		TArray<int32> Perm = 
-			Attributes.GetValueOrDefault<TArray<int32>>(TEXT("perm"), (TArray<int32>) ReversePerm);
-
+		Perm = Attributes.GetValueOrDefault<TArray<int32>>(TEXT("perm"), (TArray<int32>) ReversePerm);
 		check(Perm.Num() == NumDims);
+		return true;
 
-		// Initialize Input tensor desc
-		// Apply permutations to both sizes and strides
-		auto PermuteFunc = [&, Perm] (TConstArrayView<uint32> InputView) -> Util::FSmallUIntArray
-		{
-			Util::FSmallUIntArray Permuted;
-			
-			for (int32 PermVal : Perm)
-			{
-				Permuted.Add(InputView[PermVal]);
-			}
-			
-			return Permuted;
-		};
+	}
+
+	//
+	//
+	//
+	virtual int PrepareOutputs(TConstArrayView<NNE::Internal::FTensorRef> InputTensors, TArrayView<NNE::Internal::FTensorRef> OutputTensors) const override
+	{
+		Util::FSmallUIntArray OutputShape = Permute(InputTensors[0]->GetShape().GetData());
+
+		OutputTensors[0]->SetShape(NNE::FTensorShape::Make(OutputShape));
+
+		return 0;
+	}
+
+	//
+	//
+	//
+	virtual bool Create(IDMLDevice* Device, TConstArrayView<NNE::Internal::FTensorRef> InputTensors, TConstArrayView<NNE::Internal::FTensorRef> OutputTensors) override
+	{
+		TConstArrayView<uint32>	InputShape = InputTensors[0]->GetShape().GetData();
 
 		FTensorDescDml DmlInputTensorDesc;
 
-		DmlInputTensorDesc.SetFromTensor(InputTensors[0]);
-		DmlInputTensorDesc.SetShape(PermuteFunc(DmlInputTensorDesc.GetSizes()));
+		DmlInputTensorDesc
+			.SetFromTensor(*InputTensors[0])
+			.SetShape(Permute(InputShape))
+			.SetStridesFromShape(InputShape)
+		;
 
-		DmlInputTensorDesc.SetStridesFromShape(InputTensors[0].GetShape());
-		DmlInputTensorDesc.SetStrides(PermuteFunc(DmlInputTensorDesc.GetStrides()));
+		DmlInputTensorDesc.SetStrides(Permute(DmlInputTensorDesc.GetStrides()));
 
 		if (!DmlInputTensorDesc.Validate())
 		{
@@ -82,11 +112,10 @@ public:
 			return false;
 		}
 
-		// Initialize Output tensor desc
 		FTensorDescDml DmlOutputTensorDesc;
 
 		if (!DmlOutputTensorDesc
-				.SetFromTensor(OutputTensors[0])
+				.SetFromTensor(*OutputTensors[0])
 				.Validate())
 		{
 			UE_LOG(LogNNE, Error, TEXT("Failed to initialize Transpose output for DML inference"));
@@ -100,7 +129,7 @@ public:
 		DmlIdentityOpDesc.InputTensor = DmlInputTensorDesc.GetDmlDesc();
 		DmlIdentityOpDesc.OutputTensor = DmlOutputTensorDesc.GetDmlDesc();
 
-		return CreateOperator(Device, DML_OPERATOR_DESC{ DML_OPERATOR_ELEMENT_WISE_IDENTITY, &DmlIdentityOpDesc} );
+		return CreateOperator(Device, DML_OPERATOR_DESC{ DML_OPERATOR_ELEMENT_WISE_IDENTITY, &DmlIdentityOpDesc });
 	}
 };
 
diff --git a/Engine/Plugins/Experimental/NNERuntimeRDG/Source/NNERuntimeRDG/Private/Dml/NNEDmlOperatorUnsqueeze.cpp b/Engine/Plugins/Experimental/NNERuntimeRDG/Source/NNERuntimeRDG/Private/Dml/NNEDmlOperatorUnsqueeze.cpp
--- a/Engine/Plugins/Experimental/NNERuntimeRDG/Source/NNERuntimeRDG/Private/Dml/NNEDmlOperatorUnsqueeze.cpp
+++ b/Engine/Plugins/Experimental/NNERuntimeRDG/Source/NNERuntimeRDG/Private/Dml/NNEDmlOperatorUnsqueeze.cpp
@@ -7,17 +7,26 @@
 
 namespace UE::NNERuntimeRDG::Private::Dml
 {
-
+//
+//
+//
 class FOperatorDmlUnsqueeze : public FOperatorDml
 {
-
+	TArray<int32> Axes;
+	
 public:
 
+	//
+	//
+	//
 	static FOperatorDml* Create()
 	{
 		return new FOperatorDmlUnsqueeze();
 	}
 
+	//
+	//
+	//
 	static bool Validate(const NNE::FAttributeMap& AttributeMap, TConstArrayView<ENNETensorDataType> InputTypes, TConstArrayView<NNE::FSymbolicTensorShape> InputShapes)
 	{
 		//TODO
@@ -27,10 +36,10 @@ public:
 	//
 	//
 	//
-	virtual bool Initialize(IDMLDevice* Device, TArrayView<const NNE::Internal::FTensor> InputTensors, TArrayView<const NNE::Internal::FTensor> OutputTensors, const NNE::FAttributeMap& Attributes) override
+	virtual bool Initialize(TConstArrayView<NNE::FTensorDesc> Inputs, TConstArrayView<NNE::FTensorDesc> Outputs, const NNE::FAttributeMap& Attributes) override
 	{
-		check(InputTensors.Num() == 1);
-		check(OutputTensors.Num() == 1);
+		check(Inputs.Num() == 1);
+		check(Outputs.Num() == 1);
 
 		const FNNEAttributeValue* AxesAttr = Attributes.GetAttributeValue(TEXT("axes"));
 		if (!AxesAttr)
@@ -39,18 +48,28 @@ public:
 			return false;
 		}
 
-		TArray<int32>			Axes = AxesAttr->GetValue<TArray<int32>>();
-		TConstArrayView<uint32>	InputShape = InputTensors[0].GetShape().GetData();
-		const int32_t			OutShapeRank = InputShape.Num() + Axes.Num();
+		Axes = AxesAttr->GetValue<TArray<int32>>();
+		
+		const int32 OutShapeRank = Inputs[0].GetShape().Rank() + Axes.Num();
 
 		HandleNegativeAxes(Axes, OutShapeRank);
 		Axes.Sort();
 
-		Util::FSmallUIntArray	OutputShape;
+		return true;
+	}
+
+	//
+	//
+	//
+	virtual int PrepareOutputs(TConstArrayView<NNE::Internal::FTensorRef> InputTensors, TArrayView<NNE::Internal::FTensorRef> OutputTensors) const override
+	{
+		TConstArrayView<uint32>	InputShape = InputTensors[0]->GetShape().GetData();
+		const int32_t			OutShapeRank = InputShape.Num() + Axes.Num();
 
+		Util::FSmallUIntArray	OutputShape;
 		OutputShape.Reserve(OutShapeRank);
 		OutputShape.Append(InputShape);
-		
+
 		for (int32 Idx = 0; Idx < Axes.Num(); ++Idx)
 		{
 			if (Axes[Idx] >= OutShapeRank)
@@ -62,11 +81,23 @@ public:
 			OutputShape.Insert(1, Axes[Idx]);
 		}
 
+		OutputTensors[0]->SetShape(NNE::FTensorShape::Make(OutputShape));
+
+		return 0;
+	}
+
+	//
+	//
+	//
+	virtual bool Create(IDMLDevice* Device, TConstArrayView<NNE::Internal::FTensorRef> InputTensors, TConstArrayView<NNE::Internal::FTensorRef> OutputTensors) override
+	{
+		const NNE::Internal::FTensor& InputTensor = *InputTensors[0];
+		const NNE::Internal::FTensor& OutputTensor = *OutputTensors[0];
+
 		FTensorDescDml DmlTensorDesc;
 		
 		if (!DmlTensorDesc
-				.SetFromTensor(OutputTensors[0])
-				.SetShape(OutputShape)
+				.SetFromTensor(OutputTensor)
 				.Validate())
 		{
 			UE_LOG(LogNNE, Warning, TEXT("Failed to initialize Unsqueeze's output tensor for DML inference"));
diff --git a/Engine/Plugins/Experimental/NNERuntimeRDG/Source/NNERuntimeRDG/Private/Dml/NNEDmlOperatorUtils.h b/Engine/Plugins/Experimental/NNERuntimeRDG/Source/NNERuntimeRDG/Private/Dml/NNEDmlOperatorUtils.h
--- a/Engine/Plugins/Experimental/NNERuntimeRDG/Source/NNERuntimeRDG/Private/Dml/NNEDmlOperatorUtils.h
+++ b/Engine/Plugins/Experimental/NNERuntimeRDG/Source/NNERuntimeRDG/Private/Dml/NNEDmlOperatorUtils.h
@@ -13,192 +13,473 @@
 namespace UE::NNERuntimeRDG::Private::Dml
 {
 
+//
+//
+//
 enum EAutoPad
 {
-    NOTSET,
-    SAME_UPPER,
-    SAME_LOWER,
-    VALID
+	NOTSET,
+	SAME_UPPER,
+	SAME_LOWER,
+	VALID
 };
 
+//
+//
+//
 static EAutoPad AutoPadFromString(FStringView StringVal) 
 {
-    if (FCString::Stricmp(StringVal.GetData(), TEXT("NOTSET")) == 0) 
-    {
-        return EAutoPad::NOTSET;
-    }
-    else if (FCString::Stricmp(StringVal.GetData(), TEXT("SAME_UPPER")) == 0)
-    {
-        return EAutoPad::SAME_UPPER;
-    }
-    else if (FCString::Stricmp(StringVal.GetData(), TEXT("SAME_LOWER")) == 0)
-    {
-        return EAutoPad::SAME_LOWER;
-    }
-    else if (FCString::Stricmp(StringVal.GetData(), TEXT("VALID")) == 0)
-    {
-        return EAutoPad::VALID;
-    }
-    else
-    {
-        return EAutoPad::NOTSET;
-    }
+	if (FCString::Stricmp(StringVal.GetData(), TEXT("NOTSET")) == 0) 
+	{
+		return EAutoPad::NOTSET;
+	}
+	else if (FCString::Stricmp(StringVal.GetData(), TEXT("SAME_UPPER")) == 0)
+	{
+		return EAutoPad::SAME_UPPER;
+	}
+	else if (FCString::Stricmp(StringVal.GetData(), TEXT("SAME_LOWER")) == 0)
+	{
+		return EAutoPad::SAME_LOWER;
+	}
+	else if (FCString::Stricmp(StringVal.GetData(), TEXT("VALID")) == 0)
+	{
+		return EAutoPad::VALID;
+	}
+	else
+	{
+		return EAutoPad::NOTSET;
+	}
 }
 
+//
+//
+//
 template<typename T>
 static bool IsEqualOrBroadcastable(TConstArrayView<T> ShapeA, TConstArrayView<T> ShapeB)
 {
-    if(ShapeA.Num() < ShapeB.Num())
-    {
-        return false;
-    }
-
-    int32 BIdx = 0;
-
-    for(int32 Idx = 0; Idx < ShapeA.Num() && BIdx < ShapeB.Num(); ++Idx)
-    {
-        if(BIdx != 0)
-        {
-            if(ShapeA[Idx] != ShapeB[BIdx])
-            {
-                if(ShapeB[BIdx] != 1)
-                {
-                    return false;
-                }
-            }
-            ++BIdx;
-        }
-        if(BIdx == 0 && ShapeA[Idx] == ShapeB[BIdx])
-        {
-            ++BIdx;
-        }
-    }
-
-    if(BIdx == 0)
-    {
-        return false;
-    }
-
-    return true;
+	if (ShapeA.Num() < ShapeB.Num())
+	{
+		return false;
+	}
+
+	int32 BIdx = 0;
+
+	for(int32 Idx = 0; Idx < ShapeA.Num() && BIdx < ShapeB.Num(); ++Idx)
+	{
+		if (BIdx != 0)
+		{
+			if (ShapeA[Idx] != ShapeB[BIdx])
+			{
+				if (ShapeB[BIdx] != 1)
+				{
+					return false;
+				}
+			}
+			++BIdx;
+		}
+		if (BIdx == 0 && ShapeA[Idx] == ShapeB[BIdx])
+		{
+			++BIdx;
+		}
+	}
+
+	if (BIdx == 0)
+	{
+		return false;
+	}
+
+	return true;
 }
 
+//
+//
+//
 static bool CheckGenericTensor(ENNETensorDataType DataType, const NNE::FSymbolicTensorShape& TensorShape)
 {
-    if(!TensorShape.IsConcrete())
-    {
-        UE_LOG(LogNNE, Warning, TEXT("DML tensor shape must be concrete"));
-        return false;
-    }
-
-    if (NNE::FTensorShape::MakeFromSymbolic(TensorShape).Volume() == 0)
-    {
-        UE_LOG(LogNNE, Warning, TEXT("Invalid DML tensor size, it's 0"));
-        return false;
-    }
-
-    const int32 MinTensorRank(0), MaxTensorRank(DML_TENSOR_DIMENSION_COUNT_MAX1);
-
-    if (TensorShape.Rank() < MinTensorRank || TensorShape.Rank() > MaxTensorRank)
-    {
-        UE_LOG(LogNNE, Warning, TEXT("Invalid DML tensor rank: %d [%d,%d]"), TensorShape.Rank(), MinTensorRank, MaxTensorRank);
-        return false;
-    }
-
-    return true;
+	const int32 MinTensorRank(0), MaxTensorRank(DML_TENSOR_DIMENSION_COUNT_MAX1);
+
+	if (TensorShape.Rank() < MinTensorRank || TensorShape.Rank() > MaxTensorRank)
+	{
+		UE_LOG(LogNNE, Warning, TEXT("Invalid DML tensor rank: %d [%d,%d]"), TensorShape.Rank(), MinTensorRank, MaxTensorRank);
+		return false;
+	}
+
+	return true;
 }
 
+//
+//
+//
 static bool CheckElementwiseTensor(ENNETensorDataType DataType, const NNE::FSymbolicTensorShape& TensorShape)
 {
-    if (DataType != ENNETensorDataType::Float)
-    {
-        UE_LOG(LogNNE, Warning, TEXT("Invalid DML tensor data type"));
-        return false;
-    }
-
-    if(!CheckGenericTensor(DataType, TensorShape))
-    {
-        return false;
-    }
-
-    return true;
+	if (DataType != ENNETensorDataType::Float)
+	{
+		UE_LOG(LogNNE, Warning, TEXT("Invalid DML tensor data type"));
+		return false;
+	}
+
+	if(!CheckGenericTensor(DataType, TensorShape))
+	{
+		return false;
+	}
+
+	return true;
 }
 
+//
+//
+//
 static Util::FSmallUIntArray KernelPadding(
-    TConstArrayView<uint32> InputShape, TConstArrayView<uint32> WindowSize, 
-    TConstArrayView<uint32> Dilations, TConstArrayView<uint32> Strides
-    )
+	TConstArrayView<uint32> InputShape, TConstArrayView<uint32> WindowSize, 
+	TConstArrayView<uint32> Dilations, TConstArrayView<uint32> Strides
+	)
 {
-    const uint32 NumSpatialDimensions = InputShape.Num() - NonspatialDimensionCount;
-    check(NumSpatialDimensions >= 1);
-
-    Util::FSmallUIntArray Padding;
-    Padding.SetNumUninitialized(NumSpatialDimensions);
-
-    for (uint32 Dim = 0; Dim < NumSpatialDimensions; ++Dim)
-    {
-        uint32 InputLen = uint32(InputShape[Dim + NonspatialDimensionCount]);
-        uint32 StridedOutLen = (InputLen + Strides[Dim] - 1) / Strides[Dim];
-        uint32 KernelLen = 1 + (WindowSize[Dim] - 1) * Dilations[Dim];
-        uint32 Len = Strides[Dim] * (StridedOutLen - 1) + KernelLen;
-        
-        Padding[Dim] = (Len <= InputLen) ? 0 : (Len - InputLen);
-    }
-
-    return Padding;
+	const uint32 NumSpatialDimensions = InputShape.Num() - NonspatialDimensionCount;
+	check(NumSpatialDimensions >= 1);
+
+	Util::FSmallUIntArray Padding;
+	Padding.SetNumUninitialized(NumSpatialDimensions);
+
+	for (uint32 Dim = 0; Dim < NumSpatialDimensions; ++Dim)
+	{
+		uint32 InputLen = uint32(InputShape[Dim + NonspatialDimensionCount]);
+		uint32 StridedOutLen = (InputLen + Strides[Dim] - 1) / Strides[Dim];
+		uint32 KernelLen = 1 + (WindowSize[Dim] - 1) * Dilations[Dim];
+		uint32 Len = Strides[Dim] * (StridedOutLen - 1) + KernelLen;
+
+		Padding[Dim] = (Len <= InputLen) ? 0 : (Len - InputLen);
+	}
+
+	return Padding;
 }
 
-static void ComputeStartEndPaddings(
-    TConstArrayView<uint32> InputShape,
-    const NNE::FAttributeMap& Attributes, 
-    Util::FSmallUIntArray &OutStartPadding, 
-	Util::FSmallUIntArray &OutEndPadding,
-    TConstArrayView<uint32> Padding)
+//
+// Used by FKernelArgs
+//
+class FPaddingsHelper
 {
-    const uint32 NumSpatialDimensions = InputShape.Num() - NonspatialDimensionCount;
-    check(NumSpatialDimensions >= 1);
-
-    EAutoPad AutoPad = AutoPadFromString(*Attributes.GetValue<FString>(TEXT("auto_pad")));
-
-    if (AutoPad == EAutoPad::NOTSET)
-    {
-        Util::FSmallUIntArray Pads;
-
-        Pads.Init(0u, 2 * NumSpatialDimensions);
-        check(Util::GetArrayAttributeNoOverflow(Attributes.GetAttributeValue(TEXT("pads")), Pads, TConstArrayView<uint32>(Pads)))
-
-
-        for (uint32 Dim = 0; Dim < NumSpatialDimensions; ++Dim)
-        {
-            OutStartPadding.Add(Pads[Dim]);
-            OutEndPadding.Add(Pads[Dim + NumSpatialDimensions]);
-        }
-    }
-    else if (AutoPad == EAutoPad::VALID)
-    {
-        OutStartPadding.Init(0, NumSpatialDimensions);
-        OutEndPadding.Init(0, NumSpatialDimensions);
-    }
-    else
-    {
-        OutStartPadding.Init(0, NumSpatialDimensions);
-        OutEndPadding.Init(0, NumSpatialDimensions);
-
-        for (uint32 Dim = 0; Dim < NumSpatialDimensions; ++Dim)
-        {
-            if (AutoPad == EAutoPad::SAME_LOWER)
-            {
-                OutStartPadding[Dim] = (Padding[Dim] + 1) / 2;
-            }
-            else
-            {
-                OutStartPadding[Dim] = Padding[Dim] / 2;
-            }
-
-            OutEndPadding[Dim] = Padding[Dim] - OutStartPadding[Dim];
-        }
-    }
-}
+	Util::FSmallUIntArray	Pads;
+	EAutoPad				AutoPad;
+	uint32					NumSpatialDimensions;
+
+public:
+
+	//
+	//
+	//
+	bool Init(const NNE::FAttributeMap& Attributes, uint32 InputRank)
+	{
+		NumSpatialDimensions = InputRank - NonspatialDimensionCount;
+		check(NumSpatialDimensions >= 1);
+
+		const FNNEAttributeValue* AttrAutoPad = Attributes.GetAttributeValue(TEXT("auto_pad"));
+
+		if (AttrAutoPad)
+		{
+			AutoPad = AutoPadFromString(AttrAutoPad->GetValue<FString>());
+		}
+		else
+		{
+			AutoPad = EAutoPad::NOTSET;
+		}
+
+		if (AutoPad == EAutoPad::NOTSET)
+		{
+			Pads.Init(0u, 2 * NumSpatialDimensions);
+			if (!Util::GetArrayAttributeNoOverflow(Attributes.GetAttributeValue(TEXT("pads")), Pads, TConstArrayView<uint32>(Pads)))
+			{
+				UE_LOG(LogNNE, Error, TEXT("Pads attribute cast led to overflow"));
+				return false;
+			}
+		}
+		
+		return true;
+	}
+
+	//
+	//
+	//
+	bool Evaluate(Util::FSmallUIntArray& OutStartPadding, Util::FSmallUIntArray& OutEndPadding, TConstArrayView<uint32> Padding)
+	{
+		if (AutoPad == EAutoPad::NOTSET)
+		{
+			for (uint32 Dim = 0; Dim < NumSpatialDimensions; ++Dim)
+			{
+				OutStartPadding.Add(Pads[Dim]);
+				OutEndPadding.Add(Pads[Dim + NumSpatialDimensions]);
+			}
+		}
+		else if (AutoPad == EAutoPad::VALID)
+		{
+			OutStartPadding.Init(0, NumSpatialDimensions);
+			OutEndPadding.Init(0, NumSpatialDimensions);
+		}
+		else
+		{
+			OutStartPadding.Init(0, NumSpatialDimensions);
+			OutEndPadding.Init(0, NumSpatialDimensions);
+
+			for (uint32 Dim = 0; Dim < NumSpatialDimensions; ++Dim)
+			{
+				if (AutoPad == EAutoPad::SAME_LOWER)
+				{
+					OutStartPadding[Dim] = (Padding[Dim] + 1) / 2;
+				}
+				else
+				{
+					OutStartPadding[Dim] = Padding[Dim] / 2;
+				}
+
+				OutEndPadding[Dim] = Padding[Dim] - OutStartPadding[Dim];
+			}
+		}
+
+		return true;
+	}
+};
+
+//
+// This is a base class that is used for Conv, ConvTranspose, Pool (both local and global) operators and MaxUnpool
+//
+class FKernelArgs
+{
+
+public:
+
+	//
+	//
+	//
+	uint32 GetNumDimensions() const
+	{
+		return NumDimensions;
+	}
+
+	//
+	//
+	//
+	TConstArrayView<uint32> GetStrides() const
+	{
+		return Strides;
+	}
+
+	//
+	//
+	//
+	TConstArrayView<uint32> GetDilations() const
+	{
+		return Dilations;
+	}
+
+	//
+	//
+	//
+	TConstArrayView<uint32> GetStartPadding()
+	{
+		return StartPadding;
+	}
+
+	//
+	//
+	//
+	TConstArrayView<uint32> GetEndPadding()
+	{
+		return EndPadding;
+	}
+
+	//
+	//
+	//
+	TConstArrayView<uint32> GetOutputShape() const
+	{
+		return OutputShape;
+	}
+
+private:
+
+	FPaddingsHelper			Paddings;
+
+protected:
+
+	Util::FSmallUIntArray	StartPadding;
+	Util::FSmallUIntArray	EndPadding;
+	Util::FSmallUIntArray	OutPadding;
+	Util::FSmallUIntArray	Dilations;
+	Util::FSmallUIntArray	Strides;
+	Util::FSmallUIntArray	InOutputShape;
+	Util::FSmallUIntArray	WindowSize;
+	Util::FSmallUIntArray	OutputShape;
+	uint32					NumDimensions;
+	bool					bIsGlobalKernel;
+	bool					bIsTransposed;
+	bool					bHasOutputShape;
+
+protected:
+
+	//
+	//
+	//
+	bool Init(const NNE::FAttributeMap& Attributes, int32 InputShapeRank, bool bInIsGlobalKernel, bool bInIsTransposed)
+	{
+		check(InputShapeRank > NonspatialDimensionCount);
+		NumDimensions = InputShapeRank - NonspatialDimensionCount;
+		bIsGlobalKernel = bInIsGlobalKernel;
+		bIsTransposed = bInIsTransposed;
+		bHasOutputShape = false;
+
+		if (bIsGlobalKernel)
+		{
+			Strides.Init(1, NumDimensions);
+			Dilations.Init(1, NumDimensions);
+			StartPadding.Init(0, NumDimensions);
+			EndPadding.Init(0, NumDimensions);
+			OutPadding.Init(0, NumDimensions);
+
+			// NOTE: WindowSize needs to be set recomputed in the sub-class
+		}
+		else
+		{
+			const FNNEAttributeValue* AttrStrides = Attributes.GetAttributeValue(TEXT("strides"));
+
+			if (AttrStrides)
+			{
+				if (!Util::GetArrayAttributeNoOverflow(Attributes.GetAttributeValue(TEXT("strides")), Strides, TConstArrayView<uint32>(Strides)))
+				{
+					UE_LOG(LogNNE, Error, TEXT("Strides attribute cast led to overflow"));
+					return false;
+				}
+			}
+			else
+			{
+				Strides.Init(1, NumDimensions);
+			}
+
+			const FNNEAttributeValue* AttrDilations = Attributes.GetAttributeValue(TEXT("dilations"));
+
+			if (AttrDilations)
+			{
+				if (!Util::GetArrayAttributeNoOverflow(Attributes.GetAttributeValue(TEXT("dilations")), Dilations, TConstArrayView<uint32>(Dilations)))
+				{
+					UE_LOG(LogNNE, Error, TEXT("Dilations attribute cast led to overflow"));
+					return false;
+				}
+			}
+			else
+			{
+				Dilations.Init(1, NumDimensions);
+			}
+
+			if (bIsTransposed)
+			{
+				const FNNEAttributeValue* AttrOutPadding = Attributes.GetAttributeValue(TEXT("output_padding"));
+
+				if (AttrOutPadding)
+				{
+					if (!Util::GetArrayAttributeNoOverflow(Attributes.GetAttributeValue(TEXT("output_padding")), OutPadding, TConstArrayView<uint32>(OutPadding)))
+					{
+						UE_LOG(LogNNE, Error, TEXT("output_padding attribute cast led to overflow"));
+						return false;
+					}
+				}			
+				else
+				{
+					OutPadding.Init(0, NumDimensions);
+				}
+
+				const FNNEAttributeValue* AttrOutShape = Attributes.GetAttributeValue(TEXT("output_shape"));
+
+				if (AttrOutShape)
+				{
+					if (!Util::GetArrayAttributeNoOverflow(Attributes.GetAttributeValue(TEXT("output_shape")), InOutputShape, TConstArrayView<uint32>(InOutputShape)))
+					{
+						UE_LOG(LogNNE, Error, TEXT("output_shape attribute cast led to overflow"));
+						return false;
+					}
+				}
+
+				bHasOutputShape = AttrOutShape != nullptr;
+			}
+
+			if (!Paddings.Init(Attributes, InputShapeRank))
+			{
+				return false;
+			}
+		}
+
+		return true;
+	}
+
+public:
+
+	//
+	//
+	//
+	void Evaluate(TConstArrayView<uint32> InputShape, TConstArrayView<uint32> PaddingsValue = MakeEmptyConstArrayView<uint32>())
+	{
+		if (bIsGlobalKernel && WindowSize.IsEmpty())
+		{
+			for (uint32 Dim = 0; Dim < NumDimensions; ++Dim)
+			{
+				WindowSize.Add(uint32(InputShape[InputShape.Num() - NumDimensions + Dim]));
+			}
+		}
+
+		check(!WindowSize.IsEmpty());
+
+		if (!bIsGlobalKernel)
+		{
+			Paddings.Evaluate(StartPadding, EndPadding, PaddingsValue);
+		}
+
+		check(uint32(InputShape.Num()) >= NumDimensions);
+		const uint32 DimOffset = InputShape.Num() - NumDimensions;
+		
+		OutputShape.Reset(InputShape.Num());
+		
+		if (!bIsTransposed)
+		{
+			OutputShape.Append(InputShape.GetData(), InputShape.Num());
+
+			for (uint32 Dim = 0; Dim < NumDimensions; ++Dim)
+			{
+				uint32 InputLen = InputShape.GetData()[Dim + DimOffset];
+				uint32 PaddedLen = InputLen + StartPadding[Dim] + EndPadding[Dim];
+				uint32 KernelLen = 1 + (WindowSize[Dim] - 1) * Dilations[Dim];
+
+				checkf(KernelLen <= PaddedLen, TEXT("KernelLen must < PaddedLen"));
+				checkf(Strides[Dim] != 0, TEXT("Strides must be != 0"));
+
+				uint32 StridableOutLen = PaddedLen - KernelLen;
+				uint32 OutLen = 1 + (StridableOutLen / Strides[Dim]);
+
+				OutputShape[Dim + DimOffset] = OutLen;
+			}
+		}
+		else
+		{
+			if (bHasOutputShape)
+			{
+				OutputShape.Append(InOutputShape.GetData(), InOutputShape.Num());
+			}
+			else
+			{
+				OutputShape.Append(InputShape.GetData(), InputShape.Num());
+
+				for (uint32 Dim = 0; Dim < NumDimensions; ++Dim)
+				{
+					uint32 Padding = StartPadding[Dim] + EndPadding[Dim];
+					uint32 KernelLen = 1 + (WindowSize[Dim] - 1) * Dilations[Dim];
+
+					OutputShape[Dim + DimOffset] = (InputShape[Dim + DimOffset] - 1) * Strides[Dim] + KernelLen + OutPadding[Dim] - Padding;
+				}
+			}
+		}
+	}
+};
+
 
+//
+//
+//
 inline int32 HandleNegativeAxis(int32 Axis, int32 Rank)
 {
 	if (Axis < 0)
@@ -210,6 +491,9 @@ inline int32 HandleNegativeAxis(int32 Axis, int32 Rank)
 	return Axis;
 }
 
+//
+//
+//
 inline void HandleNegativeAxes(TArrayView<int32> Axes, int32 Rank)
 {
 	for (int32& Axis : Axes)
@@ -218,6 +502,9 @@ inline void HandleNegativeAxes(TArrayView<int32> Axes, int32 Rank)
 	}
 }
 
+//
+//
+//
 inline int32 GetDmlAxis(int32 OnnxAxis, int32 OnnxDim, int32 DmlDim)
 {
 	check(DmlDim >= OnnxDim);
@@ -227,6 +514,9 @@ inline int32 GetDmlAxis(int32 OnnxAxis, int32 OnnxDim, int32 DmlDim)
 	return DmlAxis;
 }
 
+//
+//
+//
 inline void SetDmlAxesFromOnnx(Util::FSmallUIntArray& DmlAxes, int32 Rank, TConstArrayView<int32> OnnxAxes)
 {
 	DmlAxes.Reset();
