Fixed compilation issue introduced by using std::clamp.

#rnx

diff --git a/Engine/Plugins/Experimental/NNERuntimeRDG/Source/NNERuntimeRDG/Private/Dml/NNEDmlOperatorMaxUnpool.cpp b/Engine/Plugins/Experimental/NNERuntimeRDG/Source/NNERuntimeRDG/Private/Dml/NNEDmlOperatorMaxUnpool.cpp
--- a/Engine/Plugins/Experimental/NNERuntimeRDG/Source/NNERuntimeRDG/Private/Dml/NNEDmlOperatorMaxUnpool.cpp
+++ b/Engine/Plugins/Experimental/NNERuntimeRDG/Source/NNERuntimeRDG/Private/Dml/NNEDmlOperatorMaxUnpool.cpp
@@ -119,7 +119,12 @@ public:
 
 			for (int32 Idx = 0; Idx < Shape.Num(); ++Idx)
 			{
-				OutputShape.Add(static_cast<uint32>(std::clamp<int64>(Shape[Idx], INT32_MIN, INT32_MAX)));
+				int64 ClampedVal = Shape[Idx];
+				
+				ClampedVal = ClampedVal < INT32_MIN ? INT32_MIN : ClampedVal;
+				ClampedVal = ClampedVal > INT32_MAX ? INT32_MAX : ClampedVal;
+
+				OutputShape.Add(static_cast<uint32>(ClampedVal));
 			}
 
 			OutputTensors[0]->SetShape(NNE::FTensorShape::Make(OutputShape));
