[NNE] Add cvar for enabling/disabling DML meta commands

Problem:
Currently DirectML meta commands are disabled by default.

Solution:
Add a console variable to control the behavior

Tests:
- NNEIncubatorEditor Win64
- NNEIncubator Wn64
- Non-Unity compilation test (p4v)

#rnx

diff --git a/Engine/Plugins/Experimental/NNERuntimeRDG/Source/NNERuntimeRDG/Private/Dml/NNEDmlModel.cpp b/Engine/Plugins/Experimental/NNERuntimeRDG/Source/NNERuntimeRDG/Private/Dml/NNEDmlModel.cpp
--- a/Engine/Plugins/Experimental/NNERuntimeRDG/Source/NNERuntimeRDG/Private/Dml/NNEDmlModel.cpp
+++ b/Engine/Plugins/Experimental/NNERuntimeRDG/Source/NNERuntimeRDG/Private/Dml/NNEDmlModel.cpp
@@ -11,6 +11,14 @@
 #include "ID3D12DynamicRHI.h"
 #endif
 
+// Register CVar for enabling / disabling DirectML meta commands
+static int32 GNNEDmlMetaCommands = 0;
+static FAutoConsoleVariableRef CVarNNEDmlMetaCommands(
+	TEXT("nne.dml.MetaCommands"),
+	GNNEDmlMetaCommands,
+	TEXT("Use DirectML meta commands in the NNERuntimeRDGDml (default: 0)"),
+	ECVF_Scalability);
+
 namespace UE::NNERuntimeRDG::Private::Dml
 {
 
@@ -477,8 +485,14 @@ public:
 
 		IDMLCompiledOperator* Op = nullptr;
 		HRESULT Res;
-			
-		Res = Device1->CompileGraph(&Graph, DML_EXECUTION_FLAG_DISABLE_META_COMMANDS, DML_PPV_ARGS(&Op));
+		DML_EXECUTION_FLAGS	DmlExecFlags = DML_EXECUTION_FLAG_NONE;
+
+		if (!GNNEDmlMetaCommands)
+		{
+			DmlExecFlags = DML_EXECUTION_FLAG_DISABLE_META_COMMANDS;
+		}
+
+		Res = Device1->CompileGraph(&Graph, DmlExecFlags, DML_PPV_ARGS(&Op));
 		if (FAILED(Res))
 		{
 			UE_LOG(LogNNE, Error, TEXT("Failed to compile DML graph"));
