[NNE] Remove unwanted std functions as per coding standard requirements.

Problem:
- std::swap(), std::max() and std::iota() are used, which is not allowed.

Solution:
- Replace functions with UE equivalent, implemented own Iota() since it looks like there's no UE Algo equivalent.

Tests:
- NNEIncubatorEditor Win64
- NNEIncubator Wn64
- Non-Unity compilation test (p4v)

#rnx

diff --git a/Engine/Plugins/Experimental/NNERuntimeRDG/Source/NNERuntimeRDG/Private/Dml/NNEDmlModel.cpp b/Engine/Plugins/Experimental/NNERuntimeRDG/Source/NNERuntimeRDG/Private/Dml/NNEDmlModel.cpp
--- a/Engine/Plugins/Experimental/NNERuntimeRDG/Source/NNERuntimeRDG/Private/Dml/NNEDmlModel.cpp
+++ b/Engine/Plugins/Experimental/NNERuntimeRDG/Source/NNERuntimeRDG/Private/Dml/NNEDmlModel.cpp
@@ -80,18 +80,12 @@ bool FModelInfo::ValidateGuidAndVersion(const uint8* InGuid, const uint8* InVers
 
 #ifdef NNE_USE_DIRECTML
 
-//
-//
-//
 inline FModelInstance::FDebugName::FDebugName()
 {
 	Str[0] = '\0';
 	Length = 0;
 }
 
-//
-//
-//
 inline FModelInstance::FDebugName::FDebugName(const FString& InStr)
 {
 	FTCHARToUTF8 Conv(*InStr);
@@ -101,9 +95,6 @@ inline FModelInstance::FDebugName::FDebugName(const FString& InStr)
 	Str[Length] = '\0';
 }
 
-//
-//
-//
 inline FModelInstance::FDebugName::FDebugName(FStringView InStr)
 {
 	FTCHARToUTF8 Conv(InStr.GetData());
@@ -113,24 +104,22 @@ inline FModelInstance::FDebugName::FDebugName(FStringView InStr)
 	Str[Length] = '\0';
 }
 
-//
-//
-//
 inline const char* FModelInstance::FDebugName::Get() const
 {
 	return Str;
 }
 
-//
-//
-//
+/**
+* Utility class to help out with binding DML resources when operator is initialized and dispatched
+*/
 class FModelInstance::FBindingTable
 {
 public:
 
-	//
-	//
-	//
+	/**
+	* Initialize the binding table from the model
+	* Note: This is internal class and can only be called from FModelInstance, therefore FModelInstance is always valid
+	*/
 	bool Init(FModelInstance* InModel)
 	{
 		Model = InModel;
@@ -139,9 +128,9 @@ public:
 		return true;
 	}
 
-	//
-	//
-	//
+	/**
+	* Bind requried resources for DML operator initialization
+	*/
 	void Bind(IDMLOperatorInitializer* InOpInit, TConstArrayView<FRHIBuffer*> InputBuffers, FRHIBuffer* InPersistBuff, FRHIBuffer* InTempBuff = nullptr)
 	{
 		Reset(InOpInit);
@@ -186,6 +175,9 @@ public:
 		}
 	}
 
+	/**
+	* Bind requried resources for DML operator dispatch
+	*/ 
 	void Bind(IDMLCompiledOperator* Op, TConstArrayView<FRHIBuffer*> InputBuffers, TConstArrayView<FRHIBuffer*> OutputBuffers, FRHIBuffer* InPersistBuff = nullptr, FRHIBuffer* InTempBuff = nullptr)
 	{
 		Reset(Op);
@@ -298,9 +290,9 @@ private:
 	FModelInstance*													Model;
 };
 
-//
-//
-//
+/**
+* Helper class for building DML graph
+*/
 class FModelInstance::FGraphBuilder
 {
 private:
@@ -312,9 +304,6 @@ private:
 		Intermediate
 	};
 
-	//
-	//
-	//
 	struct FEdge
 	{
 		EEdgeType	Type;
@@ -362,9 +351,6 @@ private:
 
 public:
 
-	//
-	//
-	//
 	IDMLCompiledOperator* Compile(const FModelInstance* InModel)
 	{
 		IDMLDevice*				Device = InModel->DevCtx->Device;
@@ -504,9 +490,6 @@ public:
 
 private:
 
-	//
-	//
-	//
 	bool AddEdges(const FModelInstance* InModel)
 	{
 		Edges.Reset();
@@ -592,9 +575,6 @@ private:
 		return true;
 	}
 
-	//
-	//
-	//
 	bool AddInputEdges(TConstArrayView<int32> InTensorIndices, const FModelInstance* InModel)
 	{
 		for (int32 Idx = 0; Idx < InTensorIndices.Num(); ++Idx)
@@ -628,9 +608,6 @@ private:
 		return true;
 	}
 
-	//
-	//
-	//
 	bool AddOutputEdges(const FModelInstance* InModel)
 	{
 		for (int32 Idx = 0; Idx < InModel->OutputTensorIndices.Num(); ++Idx)
@@ -653,9 +630,6 @@ private:
 		return true;
 	}
 
-	//
-	//
-	//
 	bool AddIntermediateEdges(const FModelInstance* InModel)
 	{
 		for (int32 Idx = 0; Idx < InModel->IntermediateTensorIndices.Num(); ++Idx)
@@ -679,9 +653,6 @@ private:
 		return true;
 	}
 
-	//
-	//
-	//
 	bool AddOperators(const FModelInstance* InModel)
 	{
 		for (const FGraphOpDesc& CurrOp : InModel->GraphOperators) 
@@ -739,9 +710,6 @@ private:
 		return true;
 	}
 
-	//
-	//
-	//
 	void AddInputEdge(int32 TensorIdx)
 	{
 		Edges.Emplace(
@@ -751,9 +719,6 @@ private:
 		);
 	}
 
-	//
-	//
-	//
 	void AddOutputEdge(int32 TensorIdx)
 	{
 		Edges.Emplace(
@@ -765,9 +730,6 @@ private:
 		++NumOutputs;
 	}
 
-	//
-	//
-	//
 	void AddIntermediateEdge(int32 TensorIdx, int32 NodeSrc, int32 NodeSrcOutput)
 	{
 		Edges.Emplace(
@@ -778,9 +740,6 @@ private:
 			);
 	}
 
-	//
-	//
-	//
 	bool ConnectEdgeDst(int32 TensorIdx, int32 NodeDst, int32 NodeDstInput)
 	{
 		FEdge* StartEdge = Edges.FindByPredicate(
@@ -807,9 +766,6 @@ private:
 		return bFoundEdge;
 	}
 
-	//
-	//
-	//
 	bool ConnectEdgeSrc(int32 TensorIdx, int32 NodeSrc, int32 NodeSrcOutput)
 	{
 		FEdge* StartEdge =
@@ -845,16 +801,10 @@ private:
 	int32									NumOutputs;
 };
 
-//
-//
-//
 FModelInstance::FModelInstance()
 {
 }
 
-//
-//
-//
 FModelInstance::~FModelInstance()
 {
 #if STATS
@@ -864,9 +814,6 @@ FModelInstance::~FModelInstance()
 #endif
 }
 
-//
-//
-//
 bool FModelInstance::Init(TConstArrayView<uint8> ModelData, FDmlDeviceContext* InDevCtx)
 {
 	ConstantCPUTensorIndices.Reset();
@@ -970,24 +917,11 @@ bool FModelInstance::Init(TConstArrayView<uint8> ModelData, FDmlDeviceContext* I
 
 			ConstantCPUTensorIndices.AddUnique(TensorIdx);
 		}
-
-		// FIXME: TODO: Find a way to test/validate if tensor is constant, for now we can only do that in FOperatorDml::PrepareOutputs()
-		//for (int32 TensorIdx : ConstantCPUTensorIndices)
-		//{
-		//	if (!AllTensorRDGRefs[TensorIdx]->HasPreparedData())
-		//	{
-		//		UE_LOG(LogNNE, Error, TEXT("DirectML requires tensor %s needs to be constant"), *AllTensorRDGRefs[TensorIdx]->GetName());
-		//		return false;
-		//	}
-		//}
 	}
 
 	return true;	
 }
 
-//
-//
-//
 bool FModelInstance::InitCompiledOp()
 {
 	static constexpr EBufferUsageFlags	WeightBuffUsage = BUF_UnorderedAccess;
@@ -1013,7 +947,7 @@ bool FModelInstance::InitCompiledOp()
 	DML_BINDING_PROPERTIES InitBindProps = OpInit->GetBindingProperties();
 	DML_BINDING_PROPERTIES ExecBindProps = CompiledOp->GetBindingProperties();
 
-	DescCount = std::max(InitBindProps.RequiredDescriptorCount, ExecBindProps.RequiredDescriptorCount);
+	DescCount = FMath::Max(InitBindProps.RequiredDescriptorCount, ExecBindProps.RequiredDescriptorCount);
 	
 	D3D12_DESCRIPTOR_HEAP_DESC	HeapDesc = {};
 
@@ -1167,9 +1101,6 @@ BEGIN_SHADER_PARAMETER_STRUCT(FDmlModelDispatchPassParameters, )
 	RDG_BUFFER_ACCESS_ARRAY(OutputBuffers)
 END_SHADER_PARAMETER_STRUCT()
 
-//
-//
-//
 void FModelInstance::AddDispatchOps_RenderThread(FRDGBuilder& GraphBuilder)
 {
 #if STATS
@@ -1298,8 +1229,11 @@ void FModelInstance::AddDispatchOps_RenderThread(FRDGBuilder& GraphBuilder)
 	);
 }
 
+/**
+* Create an instance of FOperatorDml
+* Note: The IDMLOperator (member of FOperatorDml) is still not created
+*/
 //
-// Create operator
 //
 FOperatorDml* FModelInstance::OpCreate(const FString& OpName, TConstArrayView<NNE::FTensorDesc> Inputs, TConstArrayView<NNE::FTensorDesc> Outputs, const NNE::FAttributeMap& Attributes)
 {
@@ -1324,9 +1258,6 @@ FOperatorDml* FModelInstance::OpCreate(const FString& OpName, TConstArrayView<NN
 	return Op;
 }
 
-//
-//
-//
 FBufferRHIRef FModelInstance::CreateRHIBuffer(FRHICommandListImmediate& RHICmdList, uint32 Size, EBufferUsageFlags Usage, ERHIAccess Access, const TCHAR* DbgName)
 {
 	FBufferRHIRef Buff = nullptr;
@@ -1342,9 +1273,6 @@ FBufferRHIRef FModelInstance::CreateRHIBuffer(FRHICommandListImmediate& RHICmdLi
 	return Buff;
 }
 
-//
-//
-//
 int FModelInstance::PrepareTensorShapesAndData()
 {
 	check(AllTensorRDGRefs.Num() == AllSymbolicTensorDescs.Num());
@@ -1408,9 +1336,6 @@ int FModelInstance::PrepareTensorShapesAndData()
 	return 0;
 }
 
-//
-//
-//
 TSharedPtr<NNE::IModelInstanceRDG> FModel::CreateModelInstanceRDG()
 {
 	FModelInstance* ModelInstance = new FModelInstance();
@@ -1426,9 +1351,6 @@ TSharedPtr<NNE::IModelInstanceRDG> FModel::CreateModelInstanceRDG()
 	return TSharedPtr<NNE::IModelInstanceRDG>(IModelInstance);
 }
 
-//
-//
-//
 FModel::FModel(const TSharedPtr<NNE::FSharedModelData>& InModelData, FDmlDeviceContext* InDevCtx)
 	: ModelData(InModelData), DevCtx(InDevCtx)
 {
diff --git a/Engine/Plugins/Experimental/NNERuntimeRDG/Source/NNERuntimeRDG/Private/Dml/NNEDmlOperatorSlice.cpp b/Engine/Plugins/Experimental/NNERuntimeRDG/Source/NNERuntimeRDG/Private/Dml/NNEDmlOperatorSlice.cpp
--- a/Engine/Plugins/Experimental/NNERuntimeRDG/Source/NNERuntimeRDG/Private/Dml/NNEDmlOperatorSlice.cpp
+++ b/Engine/Plugins/Experimental/NNERuntimeRDG/Source/NNERuntimeRDG/Private/Dml/NNEDmlOperatorSlice.cpp
@@ -12,14 +12,32 @@
 namespace UE::NNERuntimeRDG::Private::Dml
 {
 
+namespace Util
+{
+
+/**
+* Implement std::iota() replacement for TArray
+*/
+template<typename ArrayT, class ValueT>
+void Iota(ArrayT& Array, ValueT Value)
+{
+	auto First = Array.begin();
+	auto Last = Array.end();
+
+	for (auto It = First; It != Last; ++It)
+	{
+		(*It) = Value;
+		++Value;
+	}
+}
+
+} // Util
+
 /**
  * Slice
  */
 class FOperatorDmlSlice : public FOperatorDml
 {	
-	//
-	//
-	//
 	template<typename DataType>
 	static void ComputeOffsetsSizesStrides(
 		TArrayView<const NNE::Internal::FTensorRef> InputTensors,
@@ -51,7 +69,7 @@ class FOperatorDmlSlice : public FOperatorDml
 		else
 		{
 			Axes.SetNumUninitialized(Starts.Num());
-			std::iota(Axes.begin(), Axes.end(), 0);
+			Util::Iota(Axes, 0);
 		}
 
 		if (InputTensors.Num() >= 5)
@@ -94,7 +112,7 @@ class FOperatorDmlSlice : public FOperatorDml
 
 			if (Stride < 0)
             {
-                std::swap(Start, End);
+                Swap(Start, End);
                 Start += (Start < TNumericLimits<DataType>::Max()) ? 1 : 0;
                 End += (End < TNumericLimits<DataType>::Max()) ? 1 : 0;
             }
@@ -119,26 +137,16 @@ class FOperatorDmlSlice : public FOperatorDml
 
 public:
 
-	//
-	//
-	//
 	static FOperatorDml* Create()
 	{
 		return new FOperatorDmlSlice();
 	}
 
-	//
-	//
-	//
 	static bool Validate(const NNE::FAttributeMap& AttributeMap, TConstArrayView<ENNETensorDataType> InputTypes, TConstArrayView<NNE::FSymbolicTensorShape> InputShapes)
 	{
-		//TODO
 		return true;
 	}
 
-	//
-	//
-	//
 	virtual bool Initialize(TConstArrayView<NNE::FTensorDesc> Inputs, TConstArrayView<NNE::FTensorDesc> Outputs, const NNE::FAttributeMap& Attributes) override
 	{
 		check(Inputs.Num() >= 3);
@@ -175,9 +183,6 @@ public:
 		return true;
 	}
 
-	//
-	//
-	//
 	virtual int PrepareOutputs(TConstArrayView<NNE::Internal::FTensorRef> InputTensors, TArrayView<NNE::Internal::FTensorRef> OutputTensors) const override
 	{
 		switch (InputIndexDataType)
@@ -196,12 +201,8 @@ public:
 		return 0;
 	}
 
-	//
-	//
-	//
 	virtual bool Create(IDMLDevice* Device, TConstArrayView<NNE::Internal::FTensorRef> InputTensors, TConstArrayView<NNE::Internal::FTensorRef> OutputTensors) override
 	{
-		// Initialize Input tensor desc
         FTensorDescDml DmlInputTensorDesc;
         
 		if (!DmlInputTensorDesc
@@ -212,7 +213,6 @@ public:
             return false;
         }
 
-		// Initialize Output tensor desc
         FTensorDescDml DmlOutputTensorDesc;
         
 		if (!DmlOutputTensorDesc
