[NNE] Runtime RDG code cleanup

Problem:
Not all files for RDG follow UnrealEngine coding standards, for example there are empty comments, TODO, FIXME

Solution:
Update files and to meet standard requirements

Tests:
- NNEIncubatorEditor Win64
- NNEIncubator Wn64
- Non-Unity compilation test (p4v)

#rnx

diff --git a/Engine/Plugins/Experimental/NNERuntimeRDG/Source/NNERuntimeRDG/Private/Dml/NNEDmlModel.cpp b/Engine/Plugins/Experimental/NNERuntimeRDG/Source/NNERuntimeRDG/Private/Dml/NNEDmlModel.cpp
--- a/Engine/Plugins/Experimental/NNERuntimeRDG/Source/NNERuntimeRDG/Private/Dml/NNEDmlModel.cpp
+++ b/Engine/Plugins/Experimental/NNERuntimeRDG/Source/NNERuntimeRDG/Private/Dml/NNEDmlModel.cpp
@@ -90,7 +90,7 @@ inline FModelInstance::FDebugName::FDebugName(const FString& InStr)
 {
 	FTCHARToUTF8 Conv(*InStr);
 
-	Length = Conv.Length() + 1 < Size ? Conv.Length() + 1 : Size - 1;
+	Length = FMath::Min(Conv.Length() + 1, Size - 1);
 	FCStringAnsi::Strncpy(Str, Conv.Get(), Length);
 	Str[Length] = '\0';
 }
@@ -99,7 +99,7 @@ inline FModelInstance::FDebugName::FDebugName(FStringView InStr)
 {
 	FTCHARToUTF8 Conv(InStr.GetData());
 
-	Length = Conv.Length() + 1 < Size ? Conv.Length() + 1 : Size - 1;
+	Length = FMath::Min(Conv.Length() + 1, Size - 1);
 	FCStringAnsi::Strncpy(Str, Conv.Get(), Length);
 	Str[Length] = '\0';
 }
@@ -292,6 +292,7 @@ private:
 
 /**
 * Helper class for building DML graph
+* Note: This is internal class and can only be called from FModelInstance, therefore FModelInstance is always valid
 */
 class FModelInstance::FGraphBuilder
 {
@@ -1233,8 +1234,6 @@ void FModelInstance::AddDispatchOps_RenderThread(FRDGBuilder& GraphBuilder)
 * Create an instance of FOperatorDml
 * Note: The IDMLOperator (member of FOperatorDml) is still not created
 */
-//
-//
 FOperatorDml* FModelInstance::OpCreate(const FString& OpName, TConstArrayView<NNE::FTensorDesc> Inputs, TConstArrayView<NNE::FTensorDesc> Outputs, const NNE::FAttributeMap& Attributes)
 {
 	FOperatorRegistryDml::OperatorCreateFunc CreateFn = FOperatorRegistryDml::Get()->OpFind(OpName);
@@ -1359,5 +1358,3 @@ FModel::FModel(const TSharedPtr<NNE::FSharedModelData>& InModelData, FDmlDeviceC
 #endif // NNE_USE_DIRECTML
 
 } // namespace UE::NNERuntimeRDG::Private::Dml
-
-
diff --git a/Engine/Plugins/Experimental/NNERuntimeRDG/Source/NNERuntimeRDG/Private/Dml/NNEDmlModel.h b/Engine/Plugins/Experimental/NNERuntimeRDG/Source/NNERuntimeRDG/Private/Dml/NNEDmlModel.h
--- a/Engine/Plugins/Experimental/NNERuntimeRDG/Source/NNERuntimeRDG/Private/Dml/NNEDmlModel.h
+++ b/Engine/Plugins/Experimental/NNERuntimeRDG/Source/NNERuntimeRDG/Private/Dml/NNEDmlModel.h
@@ -38,9 +38,6 @@ private:
 class FDmlDeviceContext;
 class FOperatorDml;
 
-//
-//
-//
 class FModelInstance : public FModelInstanceRDG
 {
 	class FGraphBuilder;
@@ -48,9 +45,7 @@ class FModelInstance : public FModelInstanceRDG
 	
 	friend class FGraphBuilder;
 
-	//
-	//
-	//
+	// Utility class to hold debug names used by D3D
 	class FDebugName
 	{
 		static constexpr int32 Size = 128;
@@ -69,9 +64,7 @@ class FModelInstance : public FModelInstanceRDG
 		int32	Length;
 	};
 
-	//
-	//
-	//
+	// Utility class to help with deferred DML graph building
 	struct FGraphOpDesc
 	{
 		int32			OpIndex;
@@ -101,9 +94,7 @@ private:
 	bool InitCompiledOp();
 
 	FBufferRHIRef CreateRHIBuffer(FRHICommandListImmediate& RHICmdList, uint32 Size, EBufferUsageFlags Usage, ERHIAccess Access, const TCHAR* DbgName);
-	ID3D12Resource* CreateD3D12Buffer(uint32 Size, D3D12_RESOURCE_STATES ResourceState = D3D12_RESOURCE_STATE_COMMON, D3D12_HEAP_TYPE HeapType = D3D12_HEAP_TYPE_DEFAULT, const TCHAR* DebugName = nullptr);
 
-	//NOTE: This should go into RDG
 	static constexpr int32 MaxNumInputs = 512;
 	static constexpr int32 MaxNumOutputs = 4;
 
@@ -134,9 +125,6 @@ private:
 	uint64								MemSizePersist;
 };
 
-//
-//
-//
 class FModel : public NNE::IModelRDG
 {
 public:
diff --git a/Engine/Plugins/Experimental/NNERuntimeRDG/Source/NNERuntimeRDG/Private/Dml/NNEDmlOperator.cpp b/Engine/Plugins/Experimental/NNERuntimeRDG/Source/NNERuntimeRDG/Private/Dml/NNEDmlOperator.cpp
--- a/Engine/Plugins/Experimental/NNERuntimeRDG/Source/NNERuntimeRDG/Private/Dml/NNEDmlOperator.cpp
+++ b/Engine/Plugins/Experimental/NNERuntimeRDG/Source/NNERuntimeRDG/Private/Dml/NNEDmlOperator.cpp
@@ -13,33 +13,16 @@ namespace UE::NNERuntimeRDG::Private::Dml
 namespace Util
 {
 
-//
-//
-//
 inline bool IsSameShape(TConstArrayView<uint32> Left, TConstArrayView<uint32> Right)
 {
 	return Algo::Compare(Left, Right);
 }
 
-//
-//
-//
-inline bool IsSameShape(const NNE::Internal::FTensor& Left, const NNE::Internal::FTensor& Right)
-{
-	return IsSameShape(Left.GetShape().GetData(), Right.GetShape().GetData());
-}
-
-//
-//
-//
 inline bool IsSameShape(const FTensorDescDml& Left, const FTensorDescDml& Right)
 {
 	return IsSameShape(Left.GetSizes(), Right.GetSizes());
 }
 
-//
-//
-//
 DML_TENSOR_DATA_TYPE GetTensorDataType(ENNETensorDataType DataType)
 {
 	switch (DataType)
@@ -84,9 +67,6 @@ DML_TENSOR_DATA_TYPE GetTensorDataType(ENNETensorDataType DataType)
 
 } // namespace Util 
 
-//
-//
-//
 FTensorDescDml::FTensorDescDml()
 	: MinTensorRank(1)
 	, MaxTensorRank(DML_TENSOR_DIMENSION_COUNT_MAX1)
@@ -98,9 +78,6 @@ FTensorDescDml::FTensorDescDml()
 	Desc = DML_TENSOR_DESC{};
 }
 
-//
-//
-//
 FTensorDescDml& FTensorDescDml::SetTensorRank(int32 MinRank, int32 MaxRank)
 {
 	MinTensorRank = MinRank;
@@ -108,9 +85,6 @@ FTensorDescDml& FTensorDescDml::SetTensorRank(int32 MinRank, int32 MaxRank)
 	return *this;
 }
 
-//
-//
-//
 FTensorDescDml& FTensorDescDml::SetShape(TConstArrayView<uint32> Shape)
 {
 	check(!bIsValidated);
@@ -133,9 +107,6 @@ FTensorDescDml& FTensorDescDml::SetShape(TConstArrayView<uint32> Shape)
 	return *this;
 }
 
-//
-//
-//
 FTensorDescDml& FTensorDescDml::SetShape(TConstArrayView<uint32> Shape, int32 Rank)
 {
 	check(!bIsValidated);
@@ -168,9 +139,6 @@ FTensorDescDml& FTensorDescDml::SetShape(TConstArrayView<uint32> Shape, int32 Ra
 	return *this;
 }
 
-//
-//
-//
 FTensorDescDml& FTensorDescDml::SetShape(TConstArrayView<uint32> Shape, TConstArrayView<uint32> BroadcastShape)
 {
 	check(!bIsValidated);
@@ -206,9 +174,6 @@ FTensorDescDml& FTensorDescDml::SetShape(TConstArrayView<uint32> Shape, TConstAr
 	return *this;
 }
 
-//
-//
-//
 FTensorDescDml& FTensorDescDml::SetShape1D(uint32 Dimension, int32 Rank)
 {
 	check(!bIsValidated);
@@ -236,9 +201,6 @@ FTensorDescDml& FTensorDescDml::SetShape1D(uint32 Dimension, int32 Rank)
 	return *this;
 }
 
-//
-//
-//
 FTensorDescDml& FTensorDescDml::SetStridesFromShape(TConstArrayView<uint32> Shape)
 {
 	check(!bIsValidated);
@@ -258,9 +220,6 @@ FTensorDescDml& FTensorDescDml::SetStridesFromShape(TConstArrayView<uint32> Shap
 	return *this;
 }
 
-//
-//
-//
 FTensorDescDml& FTensorDescDml::SetStrides(TConstArrayView<uint32> InStrides)
 {
 	check(!bIsValidated);
@@ -273,9 +232,6 @@ FTensorDescDml& FTensorDescDml::SetStrides(TConstArrayView<uint32> InStrides)
 	return *this;
 }
 
-//
-//
-//
 FTensorDescDml& FTensorDescDml::SetDataType(ENNETensorDataType DataType)
 {
 	check(!bIsValidated);
@@ -295,9 +251,6 @@ FTensorDescDml& FTensorDescDml::SetDataType(ENNETensorDataType DataType)
 	return *this;
 }
 
-//
-//
-//
 FTensorDescDml& FTensorDescDml::SetDataOwnedByDml(bool bSetOwnedByDml)
 {
 	check(!bIsValidated);
@@ -310,9 +263,6 @@ FTensorDescDml& FTensorDescDml::SetDataOwnedByDml(bool bSetOwnedByDml)
 	return *this;
 }
 
-//
-//
-//
 bool FTensorDescDml::Validate()
 {
 	check(!bIsValidated);
@@ -366,9 +316,6 @@ bool FTensorDescDml::Validate()
 	return bIsValidated;
 }
 
-//
-//
-//
 uint64 FTensorDescDml::CalculateBufferSize(uint64 ElemSizeInBytes)
 {
 	if (ElemSizeInBytes == 0)
@@ -407,9 +354,6 @@ uint64 FTensorDescDml::CalculateBufferSize(uint64 ElemSizeInBytes)
 	return MinSizeInBytes;
 }
 
-//
-// DirectML operator base class
-//
 TConstArrayView<int32> FOperatorDml::GetConstantCPUInputs() const
 {
 	return ConstantCPUInputs;
diff --git a/Engine/Plugins/Experimental/NNERuntimeRDG/Source/NNERuntimeRDG/Private/Dml/NNEDmlOperator.h b/Engine/Plugins/Experimental/NNERuntimeRDG/Source/NNERuntimeRDG/Private/Dml/NNEDmlOperator.h
--- a/Engine/Plugins/Experimental/NNERuntimeRDG/Source/NNERuntimeRDG/Private/Dml/NNEDmlOperator.h
+++ b/Engine/Plugins/Experimental/NNERuntimeRDG/Source/NNERuntimeRDG/Private/Dml/NNEDmlOperator.h
@@ -10,9 +10,6 @@
 #include "NNEAttributeMap.h"
 #include "NNERuntimeRDGBase.h"
 
-//
-//
-//
 #define NNE_DML_REGISTER_OP(OpName) \
 struct FDmlOperator##OpName##Registrator \
 { \
@@ -28,31 +25,22 @@ static FDmlOperator##OpName##Registrator RegisterDmlOperator##OpName;
 namespace UE::NNERuntimeRDG::Private::Dml
 {
 
-static constexpr uint32_t NcdhwDimensionCount = 5;
-static constexpr uint32_t NcdhwSpatialDimensionCount = 3;
-static constexpr uint32_t NonspatialDimensionCount = 2; // The batch and channel dimensions of NCW, NCHW, NCDHW....
+static constexpr uint32 NcdhwDimensionCount = 5;
+static constexpr uint32 NcdhwSpatialDimensionCount = 3;
+static constexpr uint32 NonspatialDimensionCount = 2;		//!< The batch and channel dimensions of NCW, NCHW, NCDHW....
 
-//
-//
-//
 template<typename T>
 inline TArrayView<T> MakeEmptyArrayView()
 {
 	return MakeArrayView(static_cast<T*>(nullptr), 0);
 }
 
-//
-//
-//
 template<typename T>
 inline TConstArrayView<T> MakeEmptyConstArrayView()
 {
 	return TConstArrayView<T>(static_cast<const T*>(nullptr), 0);
 }
 
-//
-//
-//
 class FDmlDeviceContext
 {
 public:
@@ -73,9 +61,6 @@ using FSmallArray = TArray<T, TInlineAllocator<NNE::FTensorShape::MaxRank>>;
 using FSmallIntArray = TArray<int32, TInlineAllocator<NNE::FTensorShape::MaxRank>>;
 using FSmallUIntArray = TArray<uint32, TInlineAllocator<NNE::FTensorShape::MaxRank>>;
 
-//
-//
-//
 template<typename InputType, typename OutputType>
 inline bool IsOverflowing(InputType Input)
 {
@@ -87,9 +72,6 @@ inline bool IsOverflowing(InputType Input)
 	return false;
 }
 
-//
-//
-//
 template<typename InputType, typename OutputType>
 inline bool ConvertArrayViewNoOverflow(TConstArrayView<InputType> InputView, TArrayView<OutputType>& OutputView)
 {
@@ -104,9 +86,6 @@ inline bool ConvertArrayViewNoOverflow(TConstArrayView<InputType> InputView, TAr
 	return true;
 }
 
-//
-//
-//
 template<typename OutputType, typename AllocatorType>
 inline bool GetArrayAttributeNoOverflow(
 	const FNNEAttributeValue* Attr, 
@@ -155,7 +134,6 @@ inline bool GetArrayAttributeNoOverflow(
 	return true;
 }
 
-extern bool IsSameShape(const NNE::Internal::FTensor& Left, const NNE::Internal::FTensor& Right);
 extern bool IsSameShape(const FTensorDescDml& Left, const FTensorDescDml& Right);
 extern DML_TENSOR_DATA_TYPE GetTensorDataType(ENNETensorDataType DataType);
 
@@ -307,9 +285,9 @@ private:
 	bool					bIsValidated;
 };
 
-//
-// DirectML operator base class
-//
+/**
+*  DirectML operator base class
+*/
 class FOperatorDml : public IPrepareOperator
 {
 public:
diff --git a/Engine/Plugins/Experimental/NNERuntimeRDG/Source/NNERuntimeRDG/Private/Dml/NNEDmlOperatorActivations.cpp b/Engine/Plugins/Experimental/NNERuntimeRDG/Source/NNERuntimeRDG/Private/Dml/NNEDmlOperatorActivations.cpp
--- a/Engine/Plugins/Experimental/NNERuntimeRDG/Source/NNERuntimeRDG/Private/Dml/NNEDmlOperatorActivations.cpp
+++ b/Engine/Plugins/Experimental/NNERuntimeRDG/Source/NNERuntimeRDG/Private/Dml/NNEDmlOperatorActivations.cpp
@@ -1,4 +1,5 @@
 // Copyright Epic Games, Inc. All Rights Reserved.
+
 #ifdef NNE_USE_DIRECTML
 #include "NNEDmlOperator.h"
 #include "NNEDmlOperatorUtils.h"
@@ -18,17 +19,11 @@ class FOperatorDmlActivationUnary : public FOperatorDml
 {
 public:
 
-	//
-	//
-	//
 	static FOperatorDml* Create()
 	{
 		return new FOperatorDmlActivationUnary();
 	}
 
-	//
-	//
-	//
 	static bool Validate(const NNE::FAttributeMap& AttributeMap, TConstArrayView<ENNETensorDataType> InputTypes, TConstArrayView<NNE::FSymbolicTensorShape> InputShapes)
 	{
 		if(InputShapes.Num() != 1)
@@ -39,17 +34,15 @@ public:
 		return CheckElementwiseTensor(InputTypes[0], InputShapes[0]);
 	}
 
-	//
-	//
-	//
 	virtual ~FOperatorDmlActivationUnary() = default;
 
 private:
 
-	//
-	//
-	//
-	FOperatorDmlActivationUnary() : Alpha(0.0f), Beta(0.0f), Gamma(0.0f), Axis(-1) 
+	FOperatorDmlActivationUnary() 
+		: Alpha(0.0f)
+		, Beta(0.0f)
+		, Gamma(0.0f)
+		, Axis(-1) 
 	{
 	}
 	
@@ -62,9 +55,6 @@ private:
 
 public:
 
-	//
-	//
-	//
 	virtual bool Initialize(TConstArrayView<NNE::FTensorDesc> Inputs, TConstArrayView<NNE::FTensorDesc> Outputs, const NNE::FAttributeMap& Attributes) override
 	{
 		const NNE::FSymbolicTensorShape& InputShape = Inputs[0].GetShape();
@@ -89,9 +79,6 @@ public:
 		return true;
 	}
 
-	//
-	//
-	//
 	virtual int PrepareOutputs(TConstArrayView<NNE::Internal::FTensorRef> InputTensors, TArrayView<NNE::Internal::FTensorRef> OutputTensors) const override
 	{
 		check(InputTensors.Num() == 1);
@@ -101,9 +88,6 @@ public:
 		return 0;
 	}
 
-	//
-	//
-	//
 	virtual bool Create(IDMLDevice* Device, TConstArrayView<NNE::Internal::FTensorRef> InputTensors, TConstArrayView<NNE::Internal::FTensorRef> OutputTensors) override
 	{
 		const NNE::Internal::FTensor& InputTensor = *InputTensors[0];
@@ -134,9 +118,6 @@ public:
 
 private:
 
-	//
-	//
-	//
 	template<typename OpDesc>
 	void InitDmlOpDesc(OpDesc& Desc, const DML_TENSOR_DESC& TensorDesc)
 	{
@@ -144,9 +125,6 @@ private:
 		Desc.OutputTensor = &TensorDesc;
 	}
 
-	//
-	//
-	//
 	void InitDmlOpDesc(DML_ACTIVATION_LOG_SOFTMAX1_OPERATOR_DESC& Desc, const DML_TENSOR_DESC& TensorDesc)
 	{
 		Desc.InputTensor = &TensorDesc;
@@ -155,9 +133,6 @@ private:
 		Desc.Axes = (UINT*) &Axis;
 	}
 
-	//
-	//
-	//
 	void InitDmlOpDesc(DML_ACTIVATION_SOFTMAX1_OPERATOR_DESC& Desc, const DML_TENSOR_DESC& TensorDesc)
 	{
 		Desc.InputTensor = &TensorDesc;
@@ -166,9 +141,6 @@ private:
 		Desc.Axes = (UINT*) &Axis;
 	}
 
-	//
-	//
-	//
 	void InitDmlOpDesc(DML_ACTIVATION_SOFTPLUS_OPERATOR_DESC& Desc, const DML_TENSOR_DESC& TensorDesc)
 	{
 		Desc.InputTensor = &TensorDesc;
@@ -176,9 +148,6 @@ private:
 		Desc.Steepness = 1.0f;
 	}
 
-	//
-	//
-	//
 	void InitDmlOpDesc(DML_ACTIVATION_SCALED_ELU_OPERATOR_DESC& Desc, const DML_TENSOR_DESC& TensorDesc)
 	{
 		Desc.InputTensor = &TensorDesc;
@@ -187,9 +156,6 @@ private:
 		Desc.Gamma = Gamma;
 	}
 
-	//
-	//
-	//
 	void InitDmlOpDesc(DML_ACTIVATION_ELU_OPERATOR_DESC& Desc, const DML_TENSOR_DESC& TensorDesc)
 	{
 		Desc.InputTensor = &TensorDesc;
@@ -197,9 +163,6 @@ private:
 		Desc.Alpha = Alpha;
 	}
 
-	//
-	//
-	//
 	void InitDmlOpDesc(DML_ACTIVATION_HARD_SIGMOID_OPERATOR_DESC& Desc, const DML_TENSOR_DESC& TensorDesc)
 	{
 		Desc.InputTensor = &TensorDesc;
@@ -208,9 +171,6 @@ private:
 		Desc.Beta = Beta;
 	}
 
-	//
-	//
-	//
 	void InitDmlOpDesc(DML_ACTIVATION_LEAKY_RELU_OPERATOR_DESC& Desc, const DML_TENSOR_DESC& TensorDesc)
 	{
 		Desc.InputTensor = &TensorDesc;
@@ -218,9 +178,6 @@ private:
 		Desc.Alpha = Alpha;
 	}
 
-	//
-	//
-	//
 	void InitDmlOpDesc(DML_ACTIVATION_CELU_OPERATOR_DESC& Desc, const DML_TENSOR_DESC& TensorDesc)
 	{
 		Desc.InputTensor = &TensorDesc;
@@ -228,9 +185,6 @@ private:
 		Desc.Alpha = Alpha;
 	}
 
-	//
-	//
-	//
 	void InitDmlOpDesc(DML_ACTIVATION_SHRINK_OPERATOR_DESC& Desc, const DML_TENSOR_DESC& TensorDesc)
 	{
 		Desc.InputTensor = &TensorDesc;
@@ -253,17 +207,11 @@ class FOperatorDmlActivationBinary : public FOperatorDml
 
 public:
 
-	//
-	//
-	//
 	static FOperatorDml* Create()
 	{
 		return new FOperatorDmlActivationBinary();
 	}
 
-	//
-	//
-	//
 	static bool Validate(const NNE::FAttributeMap& AttributeMap, TConstArrayView<ENNETensorDataType> InputTypes, TConstArrayView<NNE::FSymbolicTensorShape> InputShapes)
 	{
 		if(InputShapes.Num() != 2)
@@ -287,24 +235,15 @@ public:
 
 private:
 
-	//
-	//
-	//
 	FOperatorDmlActivationBinary() = default;
 
 public:
 
-	//
-	//
-	//
 	virtual bool Initialize(TConstArrayView<NNE::FTensorDesc> Inputs, TConstArrayView<NNE::FTensorDesc> Outputs, const NNE::FAttributeMap& Attributes) override
 	{
 		return true;
 	}
 
-	//
-	//
-	//
 	virtual int PrepareOutputs(TConstArrayView<NNE::Internal::FTensorRef> InputTensors, TArrayView<NNE::Internal::FTensorRef> OutputTensors) const override
 	{
 		check(InputTensors.Num() == 2);
@@ -314,9 +253,6 @@ public:
 		return 0;
 	}
 
-	//
-	//
-	//
 	virtual bool Create(IDMLDevice* Device, TConstArrayView<NNE::Internal::FTensorRef> InputTensors, TConstArrayView<NNE::Internal::FTensorRef> OutputTensors)
 	{
 		const NNE::Internal::FTensor& InputATensor = *InputTensors[0];
@@ -367,9 +303,6 @@ public:
 
 private:
 
-	//
-	//
-	//
 	template<typename OpDesc>
 	void InitDmlOpDesc(OpDesc& Desc, DML_TENSOR_DESC& LHSTensorDesc, const DML_TENSOR_DESC& RHSTensorDesc, const DML_TENSOR_DESC& OutputTensorDesc)
 	{
@@ -378,9 +311,6 @@ private:
 		Desc.OutputTensor = &OutputTensorDesc;
 	}
 
-	//
-	//
-	//
 	void InitDmlOpDesc(DML_ACTIVATION_PARAMETERIZED_RELU_OPERATOR_DESC& Desc, const DML_TENSOR_DESC& LHSTensorDesc, const DML_TENSOR_DESC& RHSTensorDesc, const DML_TENSOR_DESC& OutputTensorDesc)
 	{
 		Desc.InputTensor = &LHSTensorDesc;
@@ -401,8 +331,10 @@ struct FDmlOperator##OpName##Registrator \
 static FDmlOperator##OpName##Registrator RegisterDmlOperator##OpName;
 
 
+// Register unary activation OP without any additional params
 #define REGISTER_OP_ACTIVATION_UNARY(OpName, DmlOpName) REGISTER_OP_ACTIVATION(OpName, DmlOpName, FOperatorDmlActivationUnary)
 
+// Register unary activation OP with additional params
 #define REGISTER_OP_ACTIVATION_UNARY_PARAMS(OpName, DmlOpName, InAlpha, InBeta, InGamma) \
 template<> FOperatorDmlActivationUnary<DML_ACTIVATION_##DmlOpName##_OPERATOR_DESC, DML_OPERATOR_ACTIVATION_##DmlOpName>::FOperatorDmlActivationUnary() \
 	: Alpha(InAlpha), Beta(InBeta), Gamma(InGamma) \
@@ -411,32 +343,24 @@ template<> FOperatorDmlActivationUnary<DML_ACTIVATION_##DmlOpName##_OPERATOR_DES
 REGISTER_OP_ACTIVATION_UNARY(OpName, DmlOpName)
 
 
+// Register binary activation OP 
 #define REGISTER_OP_ACTIVATION_BINARY(OpName, DmlOpName) REGISTER_OP_ACTIVATION(OpName, DmlOpName, FOperatorDmlActivationBinary)
 
-
-
-// 	REGISTER_OP_ACTIVATION_UNARY(			OpName, 			DmlOpName )
-// 	REGISTER_OP_ACTIVATION_UNARY_PARAMS(	OpName, 			DmlOpName, 				InAlpha, 							InBeta, 	InGamma )
-
-	REGISTER_OP_ACTIVATION_UNARY(			Dropout,			IDENTITY )
-	REGISTER_OP_ACTIVATION_UNARY_PARAMS(	Celu, 				CELU, 					1.0f, 								0.0f, 		0.0f )
-	REGISTER_OP_ACTIVATION_UNARY_PARAMS(	Elu, 				ELU, 					1.0f, 								0.0f, 		0.0f )
-	REGISTER_OP_ACTIVATION_UNARY_PARAMS(	HardSigmoid, 		HARD_SIGMOID, 			0.2f, 								0.5f, 		0.0f )
-	REGISTER_OP_ACTIVATION_UNARY_PARAMS(	LeakyRelu, 			LEAKY_RELU, 			0.01f, 								0.0f, 		0.0f )
-	REGISTER_OP_ACTIVATION_UNARY(			LogSoftmax, 		LOG_SOFTMAX1 )
-	REGISTER_OP_ACTIVATION_UNARY(			Relu, 				RELU )
-	REGISTER_OP_ACTIVATION_UNARY_PARAMS(	Selu, 				SCALED_ELU, 			1.67326319217681884765625f, 		0.0f, 		1.05070102214813232421875f )
-	REGISTER_OP_ACTIVATION_UNARY(			Shrink, 			SHRINK	)
-	REGISTER_OP_ACTIVATION_UNARY(			Sigmoid, 			SIGMOID	)
-	REGISTER_OP_ACTIVATION_UNARY(			Softmax, 			SOFTMAX1 )
-	REGISTER_OP_ACTIVATION_UNARY(			Softplus, 			SOFTPLUS )
-	REGISTER_OP_ACTIVATION_UNARY(			Softsign, 			SOFTSIGN )
-
-// 	REGISTER_OP_ACTIVATION_BINARY(	OpName, 			DmlOpName )
-
-	REGISTER_OP_ACTIVATION_BINARY(	Prelu, 				PARAMETERIZED_RELU )
-
-
+REGISTER_OP_ACTIVATION_UNARY(			Dropout,			IDENTITY )
+REGISTER_OP_ACTIVATION_UNARY_PARAMS(	Celu, 				CELU, 					1.0f, 								0.0f, 		0.0f )
+REGISTER_OP_ACTIVATION_UNARY_PARAMS(	Elu, 				ELU, 					1.0f, 								0.0f, 		0.0f )
+REGISTER_OP_ACTIVATION_UNARY_PARAMS(	HardSigmoid, 		HARD_SIGMOID, 			0.2f, 								0.5f, 		0.0f )
+REGISTER_OP_ACTIVATION_UNARY_PARAMS(	LeakyRelu, 			LEAKY_RELU, 			0.01f, 								0.0f, 		0.0f )
+REGISTER_OP_ACTIVATION_UNARY(			LogSoftmax, 		LOG_SOFTMAX1 )
+REGISTER_OP_ACTIVATION_UNARY(			Relu, 				RELU )
+REGISTER_OP_ACTIVATION_UNARY_PARAMS(	Selu, 				SCALED_ELU, 			1.67326319217681884765625f, 		0.0f, 		1.05070102214813232421875f )
+REGISTER_OP_ACTIVATION_UNARY(			Shrink, 			SHRINK	)
+REGISTER_OP_ACTIVATION_UNARY(			Sigmoid, 			SIGMOID	)
+REGISTER_OP_ACTIVATION_UNARY(			Softmax, 			SOFTMAX1 )
+REGISTER_OP_ACTIVATION_UNARY(			Softplus, 			SOFTPLUS )
+REGISTER_OP_ACTIVATION_UNARY(			Softsign, 			SOFTSIGN )
+
+REGISTER_OP_ACTIVATION_BINARY(	Prelu, 				PARAMETERIZED_RELU )
 
 } // namespace UE::NNERuntimeRDG::Private::Dml
 
diff --git a/Engine/Plugins/Experimental/NNERuntimeRDG/Source/NNERuntimeRDG/Private/Dml/NNEDmlOperatorBatchNormalization.cpp b/Engine/Plugins/Experimental/NNERuntimeRDG/Source/NNERuntimeRDG/Private/Dml/NNEDmlOperatorBatchNormalization.cpp
--- a/Engine/Plugins/Experimental/NNERuntimeRDG/Source/NNERuntimeRDG/Private/Dml/NNEDmlOperatorBatchNormalization.cpp
+++ b/Engine/Plugins/Experimental/NNERuntimeRDG/Source/NNERuntimeRDG/Private/Dml/NNEDmlOperatorBatchNormalization.cpp
@@ -6,9 +6,7 @@
 
 namespace UE::NNERuntimeRDG::Private::Dml
 {
-//
-//
-//
+
 class FOperatorDmlBatchNormalization : public FOperatorDml
 {
 	static constexpr float DefaultEpsilon = 1e-5f;
@@ -27,17 +25,11 @@ class FOperatorDmlBatchNormalization : public FOperatorDml
 
 public:
 
-	//
-	//
-	//
 	static FOperatorDml* Create()
 	{
 		return new FOperatorDmlBatchNormalization();
 	}
 
-	//
-	//
-	//
 	static bool Validate(const NNE::FAttributeMap& AttributeMap, TConstArrayView<ENNETensorDataType> InputTypes, TConstArrayView<NNE::FSymbolicTensorShape> InputShapes)
 	{
 		if (InputShapes.Num() != Count)
@@ -77,9 +69,6 @@ public:
 		return true;
 	}
 
-	//
-	//
-	//
 	virtual bool Initialize(TConstArrayView<NNE::FTensorDesc> Inputs, TConstArrayView<NNE::FTensorDesc> OutputTensors, const NNE::FAttributeMap& Attributes) override
 	{
 		check(Inputs.Num() == Count);
@@ -104,9 +93,6 @@ public:
 		return true;
 	}
 
-	//
-	//
-	//
 	virtual int PrepareOutputs(TConstArrayView<NNE::Internal::FTensorRef> InputTensors, TArrayView<NNE::Internal::FTensorRef> OutputTensors) const override
 	{
 		check(InputTensors.Num() == Count);
@@ -117,9 +103,6 @@ public:
 		return 0;
 	};
 
-	//
-	//
-	//
 	virtual bool Create(IDMLDevice* Device, TConstArrayView<NNE::Internal::FTensorRef> InputTensors, TConstArrayView<NNE::Internal::FTensorRef> OutputTensors) override
 	{
 		const NNE::Internal::FTensor& InputTensor = *InputTensors[0];
@@ -175,9 +158,6 @@ public:
 
 private:
 
-	//
-	//
-	//
 	inline bool SetDmlTensorDesc(FTensorDescDml& DmlTensorDesc, const NNE::Internal::FTensor& Tensor, int32 InputRank)
 	{
 		if (Tensor.GetShape().Rank() == 1)
diff --git a/Engine/Plugins/Experimental/NNERuntimeRDG/Source/NNERuntimeRDG/Private/Dml/NNEDmlOperatorCast.cpp b/Engine/Plugins/Experimental/NNERuntimeRDG/Source/NNERuntimeRDG/Private/Dml/NNEDmlOperatorCast.cpp
--- a/Engine/Plugins/Experimental/NNERuntimeRDG/Source/NNERuntimeRDG/Private/Dml/NNEDmlOperatorCast.cpp
+++ b/Engine/Plugins/Experimental/NNERuntimeRDG/Source/NNERuntimeRDG/Private/Dml/NNEDmlOperatorCast.cpp
@@ -7,24 +7,16 @@
 
 namespace UE::NNERuntimeRDG::Private::Dml
 {
-//
-//
-//
+
 class FOperatorDmlCast : public FOperatorDml
 {
 public:
 
-	//
-	//
-	//
 	static FOperatorDml* Create()
 	{
 		return new FOperatorDmlCast();
 	}
 
-	//
-	//
-	//
 	static bool Validate(const NNE::FAttributeMap& AttributeMap, TConstArrayView<ENNETensorDataType> InputTypes, TConstArrayView<NNE::FSymbolicTensorShape> InputShapes)
 	{
 		if (InputShapes.Num() != 1)
@@ -41,9 +33,6 @@ public:
 		return true;
 	}
 
-	//
-	//
-	//
 	virtual bool Initialize(TConstArrayView<NNE::FTensorDesc> Inputs, TConstArrayView<NNE::FTensorDesc> Outputs, const NNE::FAttributeMap& Attributes) override
 	{
 		check(Inputs.Num() == 1);
@@ -78,9 +67,6 @@ public:
 		return true;
 	}
 
-	//
-	//
-	//
 	virtual int PrepareOutputs(TConstArrayView<NNE::Internal::FTensorRef> InputTensors, TArrayView<NNE::Internal::FTensorRef> OutputTensors) const override
 	{
 		check(InputTensors.Num() == 1);
@@ -90,9 +76,6 @@ public:
 		return 0;
 	};
 
-	//
-	//
-	//
 	virtual bool Create(IDMLDevice* Device, TConstArrayView<NNE::Internal::FTensorRef> InputTensors, TConstArrayView<NNE::Internal::FTensorRef> OutputTensors) override
 	{
 		const NNE::Internal::FTensor& InputTensor = *InputTensors[0];
diff --git a/Engine/Plugins/Experimental/NNERuntimeRDG/Source/NNERuntimeRDG/Private/Dml/NNEDmlOperatorConcat.cpp b/Engine/Plugins/Experimental/NNERuntimeRDG/Source/NNERuntimeRDG/Private/Dml/NNEDmlOperatorConcat.cpp
--- a/Engine/Plugins/Experimental/NNERuntimeRDG/Source/NNERuntimeRDG/Private/Dml/NNEDmlOperatorConcat.cpp
+++ b/Engine/Plugins/Experimental/NNERuntimeRDG/Source/NNERuntimeRDG/Private/Dml/NNEDmlOperatorConcat.cpp
@@ -10,26 +10,17 @@
 namespace UE::NNERuntimeRDG::Private::Dml
 {
 
-/**
- * Concat
- */
 class FOperatorDmlConcat : public FOperatorDml
 {
 	int32 Axis;
 
 public:
 
-	//
-	//
-	//
 	static FOperatorDml* Create()
 	{
 		return new FOperatorDmlConcat();
 	}
 
-	//
-	//
-	//
 	static bool Validate(const NNE::FAttributeMap& AttributeMap, TConstArrayView<ENNETensorDataType> InputTypes, TConstArrayView<NNE::FSymbolicTensorShape> InputShapes)
 	{
 		if (InputShapes.Num() == 0)
@@ -81,9 +72,6 @@ public:
 		return true;
 	}
 
-	//
-	//
-	//
 	virtual bool Initialize(TConstArrayView<NNE::FTensorDesc> Inputs, TConstArrayView<NNE::FTensorDesc> Outputs, const NNE::FAttributeMap& Attributes) override
 	{
 		check(Inputs.Num() >= 1);
@@ -105,9 +93,6 @@ public:
 		return true;
 	}
 
-	//
-	//
-	//
 	virtual int PrepareOutputs(TConstArrayView<NNE::Internal::FTensorRef> InputTensors, TArrayView<NNE::Internal::FTensorRef> OutputTensors) const override
 	{
 		TArray<uint32> OutputShape(InputTensors[0]->GetShape().GetData());
@@ -133,9 +118,6 @@ public:
 		return 0;
 	}
 
-	//
-	//
-	//
 	virtual bool Create(IDMLDevice* Device, TConstArrayView<NNE::Internal::FTensorRef> InputTensors, TConstArrayView<NNE::Internal::FTensorRef> OutputTensors) override
 	{
 		TArray<FTensorDescDml> InputDescs;
diff --git a/Engine/Plugins/Experimental/NNERuntimeRDG/Source/NNERuntimeRDG/Private/Dml/NNEDmlOperatorConvolution.cpp b/Engine/Plugins/Experimental/NNERuntimeRDG/Source/NNERuntimeRDG/Private/Dml/NNEDmlOperatorConvolution.cpp
--- a/Engine/Plugins/Experimental/NNERuntimeRDG/Source/NNERuntimeRDG/Private/Dml/NNEDmlOperatorConvolution.cpp
+++ b/Engine/Plugins/Experimental/NNERuntimeRDG/Source/NNERuntimeRDG/Private/Dml/NNEDmlOperatorConvolution.cpp
@@ -7,29 +7,21 @@
 
 namespace UE::NNERuntimeRDG::Private::Dml
 {
+
 /**
- * Convolution
- */
+* Implements both Conv and ConvTranspose operators 
+*/
 template <DML_CONVOLUTION_DIRECTION Direction>
 class FOperatorDmlConv : public FOperatorDml
 {
-	//
-	//
-	//
 	class FConvArgs : public FKernelArgs
 	{
 		int32			Group;
 
 	public:
 
-		//
-		//
-		//
 		FConvArgs() = default;
 
-		//
-		//
-		//
 		bool Init(const NNE::FTensorDesc& Input, const NNE::FTensorDesc& Filter, const NNE::FAttributeMap& Attributes)
 		{
 			if (!FKernelArgs::Init(Attributes, Input.GetShape().Rank(), /*bIsGlobalKernel*/ false, /*bIsTransposed*/ Direction != DML_CONVOLUTION_DIRECTION_FORWARD))
@@ -61,9 +53,6 @@ class FOperatorDmlConv : public FOperatorDml
 			return true;
 		}
 
-		//
-		//
-		//
 		void Evaluate(TConstArrayView<uint32> InputShape, TConstArrayView<uint32> FilterShape)
 		{
 			// NOTE: To compute paddings we need WindowSize
@@ -87,9 +76,6 @@ class FOperatorDmlConv : public FOperatorDml
 			}
 		}
 
-		//
-		//
-		//
 		TConstArrayView<uint32> GetOutPadding() const
 		{
 			return MakeArrayView((const uint32*) OutPadding.GetData(), OutPadding.Num());
@@ -102,9 +88,6 @@ class FOperatorDmlConv : public FOperatorDml
 
 	private:
 
-		//
-		//
-		//
 		Util::FSmallUIntArray ConvolutionPadding(TConstArrayView<uint32> InputShape)
 		{
 			const uint32 DimOffset = NonspatialDimensionCount;
@@ -135,26 +118,16 @@ class FOperatorDmlConv : public FOperatorDml
 
 public:
 
-	//
-	//
-	//
 	static FOperatorDml* Create()
 	{
 		return new FOperatorDmlConv();
 	}
 
-	//
-	//
-	//
 	static bool Validate(const NNE::FAttributeMap& AttributeMap, TConstArrayView<ENNETensorDataType> InputTypes, TConstArrayView<NNE::FSymbolicTensorShape> InputShapes)
 	{
-		//TODO
 		return true;
 	}
 
-	//
-	//
-	//
 	virtual bool Initialize(TConstArrayView<NNE::FTensorDesc> Inputs, TConstArrayView<NNE::FTensorDesc> Outputs, const NNE::FAttributeMap& Attributes) override
 	{
 		const NNE::FTensorDesc& InputTensor = Inputs[0];
@@ -168,9 +141,6 @@ public:
 		return true;
 	}
 
-	//
-	//
-	//
 	virtual int PrepareOutputs(TConstArrayView<NNE::Internal::FTensorRef> InputTensors, TArrayView<NNE::Internal::FTensorRef> OutputTensors) const override
 	{
 		NNE::Internal::FTensor& OutputTensor = *OutputTensors[0];
@@ -181,16 +151,12 @@ public:
 		return 0;
 	}
 
-	//
-	//
-	//
 	virtual bool Create(IDMLDevice* Device, TConstArrayView<NNE::Internal::FTensorRef> InputTensors, TConstArrayView<NNE::Internal::FTensorRef> OutputTensors) override
 	{
 		const NNE::Internal::FTensor& InputTensor = *InputTensors[0];
 		const NNE::Internal::FTensor& FilterTensor = *InputTensors[1];
 		const NNE::Internal::FTensor& OutputTensor = *OutputTensors[0];
 
-		// Initialize tensor descriptors
 		FTensorDescDml	DmlInputTensorDesc;
 		FTensorDescDml	DmlFilterTensorDesc;
 		FTensorDescDml	DmlBiasTensorDesc;
@@ -262,25 +228,16 @@ public:
 	}
 };
 
-//
-//
-//
 void RegisterConvOperator()
 {
 	FOperatorRegistryDml::Get()->OpAdd(TEXT("Conv"), FOperatorDmlConv<DML_CONVOLUTION_DIRECTION_FORWARD>::Create);
 }
 
-//
-//
-//
 void RegisterConvTransposeOperator()
 {
 	FOperatorRegistryDml::Get()->OpAdd(TEXT("ConvTranspose"), FOperatorDmlConv<DML_CONVOLUTION_DIRECTION_BACKWARD>::Create);
 }
 
-//
-//
-//
 struct FDmlOperatorConvRegistrator
 {
 	FDmlOperatorConvRegistrator()
@@ -290,9 +247,6 @@ struct FDmlOperatorConvRegistrator
 	}
 };
 
-//
-//
-//
 static FDmlOperatorConvRegistrator RegisterDmlOperatorConv;
 
 } // namespace UE::NNERuntimeRDG::Private::Dml
diff --git a/Engine/Plugins/Experimental/NNERuntimeRDG/Source/NNERuntimeRDG/Private/Dml/NNEDmlOperatorCumSum.cpp b/Engine/Plugins/Experimental/NNERuntimeRDG/Source/NNERuntimeRDG/Private/Dml/NNEDmlOperatorCumSum.cpp
--- a/Engine/Plugins/Experimental/NNERuntimeRDG/Source/NNERuntimeRDG/Private/Dml/NNEDmlOperatorCumSum.cpp
+++ b/Engine/Plugins/Experimental/NNERuntimeRDG/Source/NNERuntimeRDG/Private/Dml/NNEDmlOperatorCumSum.cpp
@@ -8,9 +8,6 @@
 namespace UE::NNERuntimeRDG::Private::Dml
 {
 
-//
-//
-//
 class FOperatorDmlCumSum : public FOperatorDml
 {
 	DML_AXIS_DIRECTION	AxisDirection;
@@ -19,26 +16,16 @@ class FOperatorDmlCumSum : public FOperatorDml
 
 public:
 
-	//
-	//
-	//
 	static FOperatorDml* Create()
 	{
 		return new FOperatorDmlCumSum();
 	}
 
-	//
-	//
-	//
 	static bool Validate(const NNE::FAttributeMap& AttributeMap, TConstArrayView<ENNETensorDataType> InputTypes, TConstArrayView<NNE::FSymbolicTensorShape> InputShapes)
 	{
-		//TODO
 		return true;
 	}
 
-	//
-	//
-	//
 	virtual bool Initialize(TConstArrayView<NNE::FTensorDesc> Inputs, TConstArrayView<NNE::FTensorDesc> Outputs, const NNE::FAttributeMap& Attributes) override
 	{
 		check(Inputs.Num() == 2);
@@ -52,9 +39,6 @@ public:
 		return true;
 	}
 
-	//
-	//
-	//
 	virtual int PrepareOutputs(TConstArrayView<NNE::Internal::FTensorRef> InputTensors, TArrayView<NNE::Internal::FTensorRef> OutputTensors) const override
 	{
 		// axis tensor must be constant
@@ -104,9 +88,6 @@ public:
 		return 0;
 	}
 
-	//
-	//
-	//
 	virtual bool Create(IDMLDevice* Device, TConstArrayView<NNE::Internal::FTensorRef> InputTensors, TConstArrayView<NNE::Internal::FTensorRef> OutputTensors) override
 	{
 		const NNE::Internal::FTensor& InputTensor = *InputTensors[0];
diff --git a/Engine/Plugins/Experimental/NNERuntimeRDG/Source/NNERuntimeRDG/Private/Dml/NNEDmlOperatorDepthToSpace.cpp b/Engine/Plugins/Experimental/NNERuntimeRDG/Source/NNERuntimeRDG/Private/Dml/NNEDmlOperatorDepthToSpace.cpp
--- a/Engine/Plugins/Experimental/NNERuntimeRDG/Source/NNERuntimeRDG/Private/Dml/NNEDmlOperatorDepthToSpace.cpp
+++ b/Engine/Plugins/Experimental/NNERuntimeRDG/Source/NNERuntimeRDG/Private/Dml/NNEDmlOperatorDepthToSpace.cpp
@@ -8,14 +8,8 @@
 namespace UE::NNERuntimeRDG::Private::Dml
 {
 
-//
-//
-//
 class FOperatorDmlDepthToSpace : public FOperatorDml
 {
-	//
-	//
-	//
 	static DML_DEPTH_SPACE_ORDER SpaceOrderFromModeString(FStringView StringVal)
 	{
 		if (FCString::Stricmp(StringVal.GetData(), TEXT("CRD")) == 0)
@@ -39,26 +33,16 @@ class FOperatorDmlDepthToSpace : public FOperatorDml
 
 public:
 
-	//
-	//
-	//
 	static FOperatorDml* Create()
 	{
 		return new FOperatorDmlDepthToSpace();
 	}
 
-	//
-	//
-	//
 	static bool Validate(const NNE::FAttributeMap& AttributeMap, TConstArrayView<ENNETensorDataType> InputTypes, TConstArrayView<NNE::FSymbolicTensorShape> InputShapes)
 	{
-		//TODO
 		return true;
 	}
 
-	//
-	//
-	//
 	virtual bool Initialize(TConstArrayView<NNE::FTensorDesc> Inputs, TConstArrayView<NNE::FTensorDesc> Outputs, const NNE::FAttributeMap& Attributes) override
 	{
 		check(Inputs.Num() == 1);
@@ -80,9 +64,6 @@ public:
 		return true;
 	}
 
-	//
-	//
-	//
 	virtual int PrepareOutputs(TConstArrayView<NNE::Internal::FTensorRef> InputTensors, TArrayView<NNE::Internal::FTensorRef> OutputTensors) const override
 	{
 		TConstArrayView<uint32>		InputShape = InputTensors[0]->GetShape().GetData();
@@ -99,9 +80,6 @@ public:
 		return 0;
 	}
 	
-	//
-	//
-	//
 	virtual bool Create(IDMLDevice* Device, TConstArrayView<NNE::Internal::FTensorRef> InputTensors, TConstArrayView<NNE::Internal::FTensorRef> OutputTensors) override
 	{
 
diff --git a/Engine/Plugins/Experimental/NNERuntimeRDG/Source/NNERuntimeRDG/Private/Dml/NNEDmlOperatorElementwise.cpp b/Engine/Plugins/Experimental/NNERuntimeRDG/Source/NNERuntimeRDG/Private/Dml/NNEDmlOperatorElementwise.cpp
--- a/Engine/Plugins/Experimental/NNERuntimeRDG/Source/NNERuntimeRDG/Private/Dml/NNEDmlOperatorElementwise.cpp
+++ b/Engine/Plugins/Experimental/NNERuntimeRDG/Source/NNERuntimeRDG/Private/Dml/NNEDmlOperatorElementwise.cpp
@@ -1,4 +1,5 @@
 // Copyright Epic Games, Inc. All Rights Reserved.
+
 #ifdef NNE_USE_DIRECTML
 #include "NNEDmlOperator.h"
 #include "NNEDmlOperatorUtils.h"
@@ -18,17 +19,11 @@ class FOperatorDmlElementWiseUnary : public FOperatorDml
 {
 public:
 
-	//
-	//
-	//
 	static FOperatorDml* Create()
 	{
 		return new FOperatorDmlElementWiseUnary();
 	}
 
-	//
-	//
-	//
 	static bool Validate(const NNE::FAttributeMap& AttributeMap, TConstArrayView<ENNETensorDataType> InputTypes, TConstArrayView<NNE::FSymbolicTensorShape> InputShapes)
 	{
 		if (InputShapes.Num() != 1)
@@ -45,9 +40,6 @@ public:
 		return true;
 	}
 
-	//
-	//
-	//
 	virtual ~FOperatorDmlElementWiseUnary() = default;
 
 private:
@@ -55,9 +47,6 @@ private:
 	float Min;
 	float Max;
 
-	//
-	//
-	//
 	FOperatorDmlElementWiseUnary() 
 		: Min(TNumericLimits<float>::Min())
 		, Max(TNumericLimits<float>::Max()) 
@@ -66,9 +55,6 @@ private:
 	
 public:
 
-	//
-	//
-	//
 	virtual bool Initialize(TConstArrayView<NNE::FTensorDesc> Inputs, TConstArrayView<NNE::FTensorDesc> Outputs, const NNE::FAttributeMap& Attributes) override
 	{
 		if constexpr (std::is_same_v<DmlElementWiseOpDescType, DML_ELEMENT_WISE_CLIP_OPERATOR_DESC>)
@@ -100,9 +86,6 @@ public:
 		return true;
 	}
 
-	//
-	//
-	//
 	virtual int PrepareOutputs(TConstArrayView<NNE::Internal::FTensorRef> InputTensors, TArrayView<NNE::Internal::FTensorRef> OutputTensors) const override
 	{
 		check(InputTensors.Num() == 1);
@@ -112,9 +95,6 @@ public:
 		return 0;
 	}
 
-	//
-	//
-	//
 	virtual bool Create(IDMLDevice* Device, TConstArrayView<NNE::Internal::FTensorRef> InputTensors, TConstArrayView<NNE::Internal::FTensorRef> OutputTensors) override
 	{
 		const NNE::Internal::FTensor& InputTensor = *InputTensors[0];
@@ -145,9 +125,6 @@ public:
 
 private:
 
-	//
-	//
-	//
 	template<typename OpDesc>
 	void InitDmlOpDesc(OpDesc& Desc, const DML_TENSOR_DESC& TensorDesc)
 	{
@@ -155,9 +132,6 @@ private:
 		Desc.OutputTensor = &TensorDesc;
 	}
 
-	//
-	//
-	//
 	void InitDmlOpDesc(DML_ELEMENT_WISE_CLIP_OPERATOR_DESC& Desc, const DML_TENSOR_DESC& TensorDesc)
 	{
 		Desc.InputTensor = &TensorDesc;
@@ -181,17 +155,11 @@ class FOperatorDmlElementWiseBinary : public FOperatorDml
 
 public:
 
-	//
-	//
-	//
 	static FOperatorDml* Create()
 	{
 		return new FOperatorDmlElementWiseBinary();
 	}
 
-	//
-	//
-	//
 	static bool Validate(const NNE::FAttributeMap& AttributeMap, TConstArrayView<ENNETensorDataType> InputTypes, TConstArrayView<NNE::FSymbolicTensorShape> InputShapes)
 	{
 		if (InputShapes.Num() != 2)
@@ -215,24 +183,15 @@ public:
 
 private:
 
-	//
-	//
-	//
 	FOperatorDmlElementWiseBinary() = default;
 
 public:
 
-	//
-	//
-	//
 	virtual bool Initialize(TConstArrayView<NNE::FTensorDesc> Inputs, TConstArrayView<NNE::FTensorDesc> Outputs, const NNE::FAttributeMap& Attributes) override
 	{
 		return true;
 	}
 
-	//
-	//
-	//
 	virtual int PrepareOutputs(TConstArrayView<NNE::Internal::FTensorRef> InputTensors, TArrayView<NNE::Internal::FTensorRef> OutputTensors) const override
 	{
 		check(InputTensors.Num() == 2);
@@ -266,9 +225,6 @@ public:
 		return 0;
 	}
 
-	//
-	//
-	//
 	virtual bool Create(IDMLDevice* Device, TConstArrayView<NNE::Internal::FTensorRef> InputTensors, TConstArrayView<NNE::Internal::FTensorRef> OutputTensors) override
 	{
 		const NNE::Internal::FTensor& InputATensor = *InputTensors[0];
@@ -318,9 +274,6 @@ public:
 
 private:
 
-	//
-	//
-	//
 	template<typename OpDesc>
 	void InitDmlOpDesc(OpDesc& Desc, const DML_TENSOR_DESC& LHSTensorDesc, const DML_TENSOR_DESC& RHSTensorDesc, const DML_TENSOR_DESC& OutputTensorDesc)
 	{
@@ -329,9 +282,6 @@ private:
 		Desc.OutputTensor = &OutputTensorDesc;
 	}
 
-	//
-	//
-	//
 	void InitDmlOpDesc(DML_ELEMENT_WISE_POW_OPERATOR_DESC& Desc, const DML_TENSOR_DESC& LHSTensorDesc, const DML_TENSOR_DESC& RHSTensorDesc, const DML_TENSOR_DESC& OutputTensorDesc)
 	{
 		Desc.InputTensor = &LHSTensorDesc;
@@ -356,42 +306,38 @@ static FDmlOperator##OpName##Registrator RegisterDmlOperator##OpName;
 
 #define REGISTER_OP_ELEMENT_WISE_BINARY(OpName, DmlOpName) REGISTER_OP_ELEMENT_WISE(OpName, DmlOpName, FOperatorDmlElementWiseBinary)
 
-// 	REGISTER_OP_ELEMENT_WISE_UNARY(	OpName, 			DmlOpName )
-	REGISTER_OP_ELEMENT_WISE_UNARY( Abs,				ABS )
-	REGISTER_OP_ELEMENT_WISE_UNARY( Acos, 				ACOS )
-	REGISTER_OP_ELEMENT_WISE_UNARY( Acosh, 				ACOSH )
-	REGISTER_OP_ELEMENT_WISE_UNARY( Asin, 				ASIN )
-	REGISTER_OP_ELEMENT_WISE_UNARY( Asinh, 				ASINH )
-	REGISTER_OP_ELEMENT_WISE_UNARY( Atan,				ATAN )
-	REGISTER_OP_ELEMENT_WISE_UNARY( Atanh, 				ATANH )
-	REGISTER_OP_ELEMENT_WISE_UNARY( Ceil,				CEIL )
-	REGISTER_OP_ELEMENT_WISE_UNARY( Clip, 				CLIP )
-	REGISTER_OP_ELEMENT_WISE_UNARY( Cos,				COS )
-	REGISTER_OP_ELEMENT_WISE_UNARY( Cosh,				COSH )
-	REGISTER_OP_ELEMENT_WISE_UNARY( Erf,				ERF )
-	REGISTER_OP_ELEMENT_WISE_UNARY( Exp, 				EXP )
-	REGISTER_OP_ELEMENT_WISE_UNARY( Floor,				FLOOR )
-	REGISTER_OP_ELEMENT_WISE_UNARY( IsInf, 				IS_INFINITY )
-	REGISTER_OP_ELEMENT_WISE_UNARY( IsNan, 				IS_NAN )
-	REGISTER_OP_ELEMENT_WISE_UNARY( Log,				LOG )
-	REGISTER_OP_ELEMENT_WISE_UNARY( Neg,				NEGATE )
-	REGISTER_OP_ELEMENT_WISE_UNARY( Reciprocal, 		RECIP )
-	REGISTER_OP_ELEMENT_WISE_UNARY( Round,				ROUND )
-	REGISTER_OP_ELEMENT_WISE_UNARY( Sign,				SIGN )
-	REGISTER_OP_ELEMENT_WISE_UNARY( Sin,				SIN )
-	REGISTER_OP_ELEMENT_WISE_UNARY( Sinh,				SINH )
-	REGISTER_OP_ELEMENT_WISE_UNARY( Sqrt,				SQRT )
-	REGISTER_OP_ELEMENT_WISE_UNARY( Tan,				TAN )
-	REGISTER_OP_ELEMENT_WISE_UNARY( Tanh,				TANH )
-
-// 	REGISTER_OP_ELEMENT_WISE_BINARY( OpName, 			DmlOpName )
-	REGISTER_OP_ELEMENT_WISE_BINARY( Add,				ADD )
-	REGISTER_OP_ELEMENT_WISE_BINARY( Div,				DIVIDE )
-	REGISTER_OP_ELEMENT_WISE_BINARY( Mul,				MULTIPLY )
-	REGISTER_OP_ELEMENT_WISE_BINARY( Pow,				POW )
-	REGISTER_OP_ELEMENT_WISE_BINARY( Sub,				SUBTRACT )
-
-
+REGISTER_OP_ELEMENT_WISE_UNARY( Abs,				ABS )
+REGISTER_OP_ELEMENT_WISE_UNARY( Acos, 				ACOS )
+REGISTER_OP_ELEMENT_WISE_UNARY( Acosh, 				ACOSH )
+REGISTER_OP_ELEMENT_WISE_UNARY( Asin, 				ASIN )
+REGISTER_OP_ELEMENT_WISE_UNARY( Asinh, 				ASINH )
+REGISTER_OP_ELEMENT_WISE_UNARY( Atan,				ATAN )
+REGISTER_OP_ELEMENT_WISE_UNARY( Atanh, 				ATANH )
+REGISTER_OP_ELEMENT_WISE_UNARY( Ceil,				CEIL )
+REGISTER_OP_ELEMENT_WISE_UNARY( Clip, 				CLIP )
+REGISTER_OP_ELEMENT_WISE_UNARY( Cos,				COS )
+REGISTER_OP_ELEMENT_WISE_UNARY( Cosh,				COSH )
+REGISTER_OP_ELEMENT_WISE_UNARY( Erf,				ERF )
+REGISTER_OP_ELEMENT_WISE_UNARY( Exp, 				EXP )
+REGISTER_OP_ELEMENT_WISE_UNARY( Floor,				FLOOR )
+REGISTER_OP_ELEMENT_WISE_UNARY( IsInf, 				IS_INFINITY )
+REGISTER_OP_ELEMENT_WISE_UNARY( IsNan, 				IS_NAN )
+REGISTER_OP_ELEMENT_WISE_UNARY( Log,				LOG )
+REGISTER_OP_ELEMENT_WISE_UNARY( Neg,				NEGATE )
+REGISTER_OP_ELEMENT_WISE_UNARY( Reciprocal, 		RECIP )
+REGISTER_OP_ELEMENT_WISE_UNARY( Round,				ROUND )
+REGISTER_OP_ELEMENT_WISE_UNARY( Sign,				SIGN )
+REGISTER_OP_ELEMENT_WISE_UNARY( Sin,				SIN )
+REGISTER_OP_ELEMENT_WISE_UNARY( Sinh,				SINH )
+REGISTER_OP_ELEMENT_WISE_UNARY( Sqrt,				SQRT )
+REGISTER_OP_ELEMENT_WISE_UNARY( Tan,				TAN )
+REGISTER_OP_ELEMENT_WISE_UNARY( Tanh,				TANH )
+
+REGISTER_OP_ELEMENT_WISE_BINARY( Add,				ADD )
+REGISTER_OP_ELEMENT_WISE_BINARY( Div,				DIVIDE )
+REGISTER_OP_ELEMENT_WISE_BINARY( Mul,				MULTIPLY )
+REGISTER_OP_ELEMENT_WISE_BINARY( Pow,				POW )
+REGISTER_OP_ELEMENT_WISE_BINARY( Sub,				SUBTRACT )
 
 } // namespace UE::NNERuntimeRDG::Private::Dml
 
diff --git a/Engine/Plugins/Experimental/NNERuntimeRDG/Source/NNERuntimeRDG/Private/Dml/NNEDmlOperatorFlatten.cpp b/Engine/Plugins/Experimental/NNERuntimeRDG/Source/NNERuntimeRDG/Private/Dml/NNEDmlOperatorFlatten.cpp
--- a/Engine/Plugins/Experimental/NNERuntimeRDG/Source/NNERuntimeRDG/Private/Dml/NNEDmlOperatorFlatten.cpp
+++ b/Engine/Plugins/Experimental/NNERuntimeRDG/Source/NNERuntimeRDG/Private/Dml/NNEDmlOperatorFlatten.cpp
@@ -21,13 +21,9 @@ public:
 
 	static bool Validate(const NNE::FAttributeMap& AttributeMap, TConstArrayView<ENNETensorDataType> InputTypes, TConstArrayView<NNE::FSymbolicTensorShape> InputShapes)
 	{
-		//TODO
 		return true;
 	}
 
-	//
-	//
-	//
 	virtual bool Initialize(TConstArrayView<NNE::FTensorDesc> Inputs, TConstArrayView<NNE::FTensorDesc> Outputs, const NNE::FAttributeMap& Attributes) override
 	{
 		check(Inputs.Num() == 1);
@@ -52,9 +48,6 @@ public:
 		return true;
 	}
 
-	//
-	//
-	//
 	virtual int PrepareOutputs(TConstArrayView<NNE::Internal::FTensorRef> InputTensors, TArrayView<NNE::Internal::FTensorRef> OutputTensors) const override
 	{
 		TConstArrayView<uint32> InputShape = InputTensors[0]->GetShape().GetData();
@@ -77,9 +70,6 @@ public:
 		return 0;
 	}
 
-	//
-	//
-	//
 	virtual bool Create(IDMLDevice * Device, TConstArrayView<NNE::Internal::FTensorRef> InputTensors, TConstArrayView<NNE::Internal::FTensorRef> OutputTensors) override
 	{
 		FTensorDescDml DmlTensorDesc;
diff --git a/Engine/Plugins/Experimental/NNERuntimeRDG/Source/NNERuntimeRDG/Private/Dml/NNEDmlOperatorGather.cpp b/Engine/Plugins/Experimental/NNERuntimeRDG/Source/NNERuntimeRDG/Private/Dml/NNEDmlOperatorGather.cpp
--- a/Engine/Plugins/Experimental/NNERuntimeRDG/Source/NNERuntimeRDG/Private/Dml/NNEDmlOperatorGather.cpp
+++ b/Engine/Plugins/Experimental/NNERuntimeRDG/Source/NNERuntimeRDG/Private/Dml/NNEDmlOperatorGather.cpp
@@ -20,13 +20,9 @@ public:
 
 	static bool Validate(const NNE::FAttributeMap& AttributeMap, TConstArrayView<ENNETensorDataType> InputTypes, TConstArrayView<NNE::FSymbolicTensorShape> InputShapes)
 	{
-		//TODO
 		return true;
 	}
 
-	//
-	//
-	//
 	virtual bool Initialize(TConstArrayView<NNE::FTensorDesc> Inputs, TConstArrayView<NNE::FTensorDesc> Outputs, const NNE::FAttributeMap& Attributes) override
 	{
 		check(Inputs.Num() == 2);
@@ -64,9 +60,6 @@ public:
 		return true;
 	}
 
-	//
-	//
-	//
 	virtual int PrepareOutputs(TConstArrayView<NNE::Internal::FTensorRef> InputTensors, TArrayView<NNE::Internal::FTensorRef> OutputTensors) const override
 	{
 		const NNE::FTensorShape& InputShape = InputTensors[0]->GetShape();
@@ -101,9 +94,6 @@ public:
 		return 0;
 	}
 
-	//
-	//
-	//
 	virtual bool Create(IDMLDevice* Device, TConstArrayView<NNE::Internal::FTensorRef> InputTensors, TConstArrayView<NNE::Internal::FTensorRef> OutputTensors) override
 	{
 		const NNE::Internal::FTensor& InputTensor	= *InputTensors[0];
diff --git a/Engine/Plugins/Experimental/NNERuntimeRDG/Source/NNERuntimeRDG/Private/Dml/NNEDmlOperatorGemm.cpp b/Engine/Plugins/Experimental/NNERuntimeRDG/Source/NNERuntimeRDG/Private/Dml/NNEDmlOperatorGemm.cpp
--- a/Engine/Plugins/Experimental/NNERuntimeRDG/Source/NNERuntimeRDG/Private/Dml/NNEDmlOperatorGemm.cpp
+++ b/Engine/Plugins/Experimental/NNERuntimeRDG/Source/NNERuntimeRDG/Private/Dml/NNEDmlOperatorGemm.cpp
@@ -6,9 +6,6 @@
 namespace UE::NNERuntimeRDG::Private::Dml
 {
 
-/**
- * Gemm
- */
 class FOperatorDmlGemm : public FOperatorDml
 {
 	float Alpha = 1.0f;
@@ -25,16 +22,11 @@ public:
 
 	static bool Validate(const NNE::FAttributeMap& AttributeMap, TConstArrayView<ENNETensorDataType> InputTypes, TConstArrayView<NNE::FSymbolicTensorShape> InputShapes)
 	{
-		//TODO
 		return true;
 	}
 
-	//
-	//
-	//
 	virtual bool Initialize(TConstArrayView<NNE::FTensorDesc> Inputs, TConstArrayView<NNE::FTensorDesc> Outputs, const NNE::FAttributeMap& Attributes) override
 	{
-		// Setup attributes
 		Alpha = Attributes.GetValueOrDefault(TEXT("alpha"), Alpha);
 		Beta = Attributes.GetValueOrDefault(TEXT("beta"), Beta);
 		TransA = Attributes.GetValueOrDefault(TEXT("transA"), TransA);
@@ -69,9 +61,6 @@ public:
 		return true;
 	}
 
-	//
-	//
-	//
 	virtual int PrepareOutputs(TConstArrayView<NNE::Internal::FTensorRef> InputTensors, TArrayView<NNE::Internal::FTensorRef> OutputTensors) const override
 	{
 		check(InputTensors.Num() >= 2 && InputTensors.Num() <= 3);
@@ -94,16 +83,12 @@ public:
 		return 0;
 	}
 
-	//
-	//
-	//
 	virtual bool Create(IDMLDevice* Device, TConstArrayView<NNE::Internal::FTensorRef> InputTensors, TConstArrayView<NNE::Internal::FTensorRef> OutputTensors) override
 	{
 		const NNE::Internal::FTensor& InputATensor = *InputTensors[0];
 		const NNE::Internal::FTensor& InputBTensor = *InputTensors[1];
 		const NNE::Internal::FTensor& OutputTensor = *OutputTensors[0];
 
-		// Initialize tensor descriptors
 		FTensorDescDml	DmlInputATensorDesc;
 		FTensorDescDml	DmlInputBTensorDesc;
 		FTensorDescDml	DmlInputCTensorDesc;
diff --git a/Engine/Plugins/Experimental/NNERuntimeRDG/Source/NNERuntimeRDG/Private/Dml/NNEDmlOperatorIdentity.cpp b/Engine/Plugins/Experimental/NNERuntimeRDG/Source/NNERuntimeRDG/Private/Dml/NNEDmlOperatorIdentity.cpp
--- a/Engine/Plugins/Experimental/NNERuntimeRDG/Source/NNERuntimeRDG/Private/Dml/NNEDmlOperatorIdentity.cpp
+++ b/Engine/Plugins/Experimental/NNERuntimeRDG/Source/NNERuntimeRDG/Private/Dml/NNEDmlOperatorIdentity.cpp
@@ -8,34 +8,21 @@
 namespace UE::NNERuntimeRDG::Private::Dml
 {
 
-//
-//
-//
 class FOperatorDmlIdentity : public FOperatorDml
 {
 
 public:
 
-	//
-	//
-	//
 	static FOperatorDml* Create()
 	{
 		return new FOperatorDmlIdentity();
 	}
 
-	//
-	//
-	//
 	static bool Validate(const NNE::FAttributeMap& AttributeMap, TConstArrayView<ENNETensorDataType> InputTypes, TConstArrayView<NNE::FSymbolicTensorShape> InputShapes)
 	{
-		//TODO
 		return true;
 	}
 
-	//
-	//
-	//
 	virtual bool Initialize(TConstArrayView<NNE::FTensorDesc> Inputs, TConstArrayView<NNE::FTensorDesc> Outputs, const NNE::FAttributeMap& Attributes) override
 	{
 		check(Inputs.Num() == 1);
@@ -44,18 +31,12 @@ public:
 		return true;
 	}
 
-	//
-	//
-	//
 	virtual int PrepareOutputs(TConstArrayView<NNE::Internal::FTensorRef> InputTensors, TArrayView<NNE::Internal::FTensorRef> OutputTensors) const override
 	{
 		OutputTensors[0]->SetShape(InputTensors[0]->GetShape());
 		return 0;
 	}
 
-	//
-	//
-	//
 	virtual bool Create(IDMLDevice* Device, TConstArrayView<NNE::Internal::FTensorRef> InputTensors, TConstArrayView<NNE::Internal::FTensorRef> OutputTensors) override
 	{
 		const NNE::Internal::FTensor& InputTensor = *InputTensors[0];
diff --git a/Engine/Plugins/Experimental/NNERuntimeRDG/Source/NNERuntimeRDG/Private/Dml/NNEDmlOperatorInstanceNormalization.cpp b/Engine/Plugins/Experimental/NNERuntimeRDG/Source/NNERuntimeRDG/Private/Dml/NNEDmlOperatorInstanceNormalization.cpp
--- a/Engine/Plugins/Experimental/NNERuntimeRDG/Source/NNERuntimeRDG/Private/Dml/NNEDmlOperatorInstanceNormalization.cpp
+++ b/Engine/Plugins/Experimental/NNERuntimeRDG/Source/NNERuntimeRDG/Private/Dml/NNEDmlOperatorInstanceNormalization.cpp
@@ -6,9 +6,6 @@
 namespace UE::NNERuntimeRDG::Private::Dml
 {
 
-//
-//
-//
 class FOperatorDmlInstanceNormalization : public FOperatorDml
 {
 	static constexpr float DefaultEpsilon = 0.00001f;
@@ -17,9 +14,6 @@ class FOperatorDmlInstanceNormalization : public FOperatorDml
 
 public:
 
-	//
-	//
-	//
 	static FOperatorDml* Create()
 	{
 		return new FOperatorDmlInstanceNormalization();
@@ -27,13 +21,9 @@ public:
 
 	static bool Validate(const NNE::FAttributeMap& AttributeMap, TConstArrayView<ENNETensorDataType> InputTypes, TConstArrayView<NNE::FSymbolicTensorShape> InputShapes)
 	{
-		//TODO
 		return true;
 	}
 
-	//
-	//
-	//
 	virtual bool Initialize(TConstArrayView<NNE::FTensorDesc> Inputs, TConstArrayView<NNE::FTensorDesc> Outputs, const NNE::FAttributeMap& Attributes) override
 	{
 		check(Inputs.Num() >= 1 && Inputs.Num() <= 3);
@@ -45,18 +35,12 @@ public:
 		return true;
 	}
 
-	//
-	//
-	//
 	virtual int PrepareOutputs(TConstArrayView<NNE::Internal::FTensorRef> InputTensors, TArrayView<NNE::Internal::FTensorRef> OutputTensors) const override
 	{
 		OutputTensors[0]->SetShape(InputTensors[0]->GetShape());
 		return 0;
 	}
 
-	//
-	//
-	//
 	virtual bool Create(IDMLDevice* Device, TConstArrayView<NNE::Internal::FTensorRef> InputTensors, TConstArrayView<NNE::Internal::FTensorRef> OutputTensors) override
 	{
 		const NNE::Internal::FTensor& InputTensor = *InputTensors[0];
diff --git a/Engine/Plugins/Experimental/NNERuntimeRDG/Source/NNERuntimeRDG/Private/Dml/NNEDmlOperatorLRN.cpp b/Engine/Plugins/Experimental/NNERuntimeRDG/Source/NNERuntimeRDG/Private/Dml/NNEDmlOperatorLRN.cpp
--- a/Engine/Plugins/Experimental/NNERuntimeRDG/Source/NNERuntimeRDG/Private/Dml/NNEDmlOperatorLRN.cpp
+++ b/Engine/Plugins/Experimental/NNERuntimeRDG/Source/NNERuntimeRDG/Private/Dml/NNEDmlOperatorLRN.cpp
@@ -7,9 +7,6 @@
 namespace UE::NNERuntimeRDG::Private::Dml
 {
 
-//
-//
-//
 class FOperatorDmlLRN : public FOperatorDml
 {
 	int32 Size;
@@ -19,26 +16,16 @@ class FOperatorDmlLRN : public FOperatorDml
 
 public:
 
-	//
-	//
-	//
 	static FOperatorDml* Create()
 	{
 		return new FOperatorDmlLRN();
 	}
 
-	//
-	//
-	//
 	static bool Validate(const NNE::FAttributeMap& AttributeMap, TConstArrayView<ENNETensorDataType> InputTypes, TConstArrayView<NNE::FSymbolicTensorShape> InputShapes)
 	{
-		//TODO
 		return true;
 	}
 
-	//
-	//
-	//
 	virtual bool Initialize(TConstArrayView<NNE::FTensorDesc> Inputs, TConstArrayView<NNE::FTensorDesc> Outputs, const NNE::FAttributeMap& Attributes) override
 	{
 		check(Inputs.Num() == 1);
@@ -53,18 +40,12 @@ public:
 		return true;
 	}
 
-	//
-	//
-	//
 	virtual int PrepareOutputs(TConstArrayView<NNE::Internal::FTensorRef> InputTensors, TArrayView<NNE::Internal::FTensorRef> OutputTensors) const override
 	{
 		OutputTensors[0]->SetShape(InputTensors[0]->GetShape());
 		return 0;
 	}
 
-	//
-	//
-	//
 	virtual bool Create(IDMLDevice* Device, TConstArrayView<NNE::Internal::FTensorRef> InputTensors, TConstArrayView<NNE::Internal::FTensorRef> OutputTensors) override
 	{
 		const NNE::Internal::FTensor& InputTensor = *InputTensors[0];
diff --git a/Engine/Plugins/Experimental/NNERuntimeRDG/Source/NNERuntimeRDG/Private/Dml/NNEDmlOperatorLpNormalization.cpp b/Engine/Plugins/Experimental/NNERuntimeRDG/Source/NNERuntimeRDG/Private/Dml/NNEDmlOperatorLpNormalization.cpp
--- a/Engine/Plugins/Experimental/NNERuntimeRDG/Source/NNERuntimeRDG/Private/Dml/NNEDmlOperatorLpNormalization.cpp
+++ b/Engine/Plugins/Experimental/NNERuntimeRDG/Source/NNERuntimeRDG/Private/Dml/NNEDmlOperatorLpNormalization.cpp
@@ -7,9 +7,6 @@
 namespace UE::NNERuntimeRDG::Private::Dml
 {
 
-//
-//
-//
 class FOperatorDmlLpNormalization : public FOperatorDml
 {
 	static constexpr float DefaultEpsilon = 1e-5f;
@@ -20,26 +17,16 @@ class FOperatorDmlLpNormalization : public FOperatorDml
 
 public:
 
-	//
-	//
-	//
 	static FOperatorDml* Create()
 	{
 		return new FOperatorDmlLpNormalization();
 	}
 
-	//
-	//
-	//
 	static bool Validate(const NNE::FAttributeMap& AttributeMap, TConstArrayView<ENNETensorDataType> InputTypes, TConstArrayView<NNE::FSymbolicTensorShape> InputShapes)
 	{
-		//TODO
 		return true;
 	}
 
-	//
-	//
-	//
 	virtual bool Initialize(TConstArrayView<NNE::FTensorDesc> Inputs, TConstArrayView<NNE::FTensorDesc> Outputs, const NNE::FAttributeMap& Attributes) override
 	{
 		check(Inputs.Num() == 1);
@@ -54,18 +41,12 @@ public:
 		return true;
 	}
 
-	//
-	//
-	//
 	virtual int PrepareOutputs(TConstArrayView<NNE::Internal::FTensorRef> InputTensors, TArrayView<NNE::Internal::FTensorRef> OutputTensors) const override
 	{
 		OutputTensors[0]->SetShape(InputTensors[0]->GetShape());
 		return 0;
 	}
 
-	//
-	//
-	//
 	virtual bool Create(IDMLDevice* Device, TConstArrayView<NNE::Internal::FTensorRef> InputTensors, TConstArrayView<NNE::Internal::FTensorRef> OutputTensors) override
 	{
 		const NNE::Internal::FTensor& InputTensor = *InputTensors[0];
diff --git a/Engine/Plugins/Experimental/NNERuntimeRDG/Source/NNERuntimeRDG/Private/Dml/NNEDmlOperatorMatMul.cpp b/Engine/Plugins/Experimental/NNERuntimeRDG/Source/NNERuntimeRDG/Private/Dml/NNEDmlOperatorMatMul.cpp
--- a/Engine/Plugins/Experimental/NNERuntimeRDG/Source/NNERuntimeRDG/Private/Dml/NNEDmlOperatorMatMul.cpp
+++ b/Engine/Plugins/Experimental/NNERuntimeRDG/Source/NNERuntimeRDG/Private/Dml/NNEDmlOperatorMatMul.cpp
@@ -6,9 +6,6 @@
 namespace UE::NNERuntimeRDG::Private::Dml
 {
 
-/**
- * MatMul
- */
 class FOperatorDmlMatMul : public FOperatorDml
 {
 public:
@@ -20,13 +17,9 @@ public:
 
 	static bool Validate(const NNE::FAttributeMap& AttributeMap, TConstArrayView<ENNETensorDataType> InputTypes, TConstArrayView<NNE::FSymbolicTensorShape> InputShapes)
 	{
-		//TODO
 		return true;
 	}
 
-	//
-	//
-	//
 	virtual bool Initialize(TConstArrayView<NNE::FTensorDesc> Inputs, TConstArrayView<NNE::FTensorDesc> Outputs, const NNE::FAttributeMap& Attributes) override
 	{
 		check(Inputs.Num() == 2);
@@ -35,9 +28,6 @@ public:
 		return true;
 	}
 
-	//
-	//
-	//
 	virtual int PrepareOutputs(TConstArrayView<NNE::Internal::FTensorRef> InputTensors, TArrayView<NNE::Internal::FTensorRef> OutputTensors) const override
 	{
 		check(InputTensors.Num() >= 2 && InputTensors.Num() <= 3);
@@ -87,9 +77,6 @@ public:
 		return 0;
 	}
 
-	//
-	//
-	//
 	virtual bool Create(IDMLDevice* Device, TConstArrayView<NNE::Internal::FTensorRef> InputTensors, TConstArrayView<NNE::Internal::FTensorRef> OutputTensors) override
 	{
 		const NNE::Internal::FTensor& ATensor = *InputTensors[0];
@@ -255,10 +242,6 @@ public:
 			DmlShapeOut[1] /= BroadcastFactor;
 		}
 		
-		// Note: DmlShape* are the the 4D ones that InitFromTensor() would need to use, however can't set them in original tensor like below because of prepared data check.
-		//ATensor.SetShape(DmlShapeA);
-		//BTensor.SetShape(DmlShapeB);
-
 		FTensorDescDml	DmlInputATensorDesc;
 		FTensorDescDml	DmlInputBTensorDesc;
 		FTensorDescDml	DmlOutputTensorDesc;
diff --git a/Engine/Plugins/Experimental/NNERuntimeRDG/Source/NNERuntimeRDG/Private/Dml/NNEDmlOperatorMaxUnpool.cpp b/Engine/Plugins/Experimental/NNERuntimeRDG/Source/NNERuntimeRDG/Private/Dml/NNEDmlOperatorMaxUnpool.cpp
--- a/Engine/Plugins/Experimental/NNERuntimeRDG/Source/NNERuntimeRDG/Private/Dml/NNEDmlOperatorMaxUnpool.cpp
+++ b/Engine/Plugins/Experimental/NNERuntimeRDG/Source/NNERuntimeRDG/Private/Dml/NNEDmlOperatorMaxUnpool.cpp
@@ -8,23 +8,14 @@
 namespace UE::NNERuntimeRDG::Private::Dml
 {
 
-//
-//
-//
 class FOperatorDmlMaxUnpool : public FOperatorDml
 {
-	//
-	//
-	//
 	class FUnpoolingArgs : public FKernelArgs
 	{
 		Util::FSmallUIntArray	KernelShape;
 
 	public:
 
-		//
-		//
-		//
 		bool Init(const NNE::FAttributeMap& Attributes, const int32 InputShapeRank)
 		{
 			if (!FKernelArgs::Init(Attributes, InputShapeRank, /*bIsGlobalKernel*/ false, /*bIsTransposed*/ true))
@@ -59,26 +50,16 @@ class FOperatorDmlMaxUnpool : public FOperatorDml
 
 public:
 
-	//
-	//
-	//
 	static FOperatorDml* Create()
 	{
 		return new FOperatorDmlMaxUnpool();
 	}
 
-	//
-	//
-	//
 	static bool Validate(const NNE::FAttributeMap& AttributeMap, TConstArrayView<ENNETensorDataType> InputTypes, TConstArrayView<NNE::FSymbolicTensorShape> InputShapes)
 	{
-		//TODO
 		return true;
 	}
 
-	//
-	//
-	//
 	virtual bool Initialize(TConstArrayView<NNE::FTensorDesc> Inputs, TConstArrayView<NNE::FTensorDesc> Outputs, const NNE::FAttributeMap& Attributes) override
 	{
 		check(Inputs.Num() == 2 || Inputs.Num() == 3);
@@ -105,9 +86,6 @@ public:
 		return true;
 	}
 
-	//
-	//
-	//
 	virtual int PrepareOutputs(TConstArrayView<NNE::Internal::FTensorRef> InputTensors, TArrayView<NNE::Internal::FTensorRef> OutputTensors) const override
 	{
 		Args.Evaluate(InputTensors[0]->GetShape().GetData());
@@ -137,9 +115,6 @@ public:
 		return 0;
 	}
 
-	//
-	//
-	//
 	virtual bool Create(IDMLDevice* Device, TConstArrayView<NNE::Internal::FTensorRef> InputTensors, TConstArrayView<NNE::Internal::FTensorRef> OutputTensors) override
 	{
 		const NNE::Internal::FTensor& InputTensor = *InputTensors[0];
diff --git a/Engine/Plugins/Experimental/NNERuntimeRDG/Source/NNERuntimeRDG/Private/Dml/NNEDmlOperatorMeanVarianceNormalization.cpp b/Engine/Plugins/Experimental/NNERuntimeRDG/Source/NNERuntimeRDG/Private/Dml/NNEDmlOperatorMeanVarianceNormalization.cpp
--- a/Engine/Plugins/Experimental/NNERuntimeRDG/Source/NNERuntimeRDG/Private/Dml/NNEDmlOperatorMeanVarianceNormalization.cpp
+++ b/Engine/Plugins/Experimental/NNERuntimeRDG/Source/NNERuntimeRDG/Private/Dml/NNEDmlOperatorMeanVarianceNormalization.cpp
@@ -7,9 +7,6 @@
 namespace UE::NNERuntimeRDG::Private::Dml
 {
 
-//
-//
-//
 class FOperatorDmlMeanVarianceNormalization : public FOperatorDml
 {
 	static constexpr float DefaultEpsilon = 1e-5f;
@@ -20,26 +17,16 @@ class FOperatorDmlMeanVarianceNormalization : public FOperatorDml
 
 public:
 
-	//
-	//
-	//
 	static FOperatorDml* Create()
 	{
 		return new FOperatorDmlMeanVarianceNormalization();
 	}
 
-	//
-	//
-	//
 	static bool Validate(const NNE::FAttributeMap& AttributeMap, TConstArrayView<ENNETensorDataType> InputTypes, TConstArrayView<NNE::FSymbolicTensorShape> InputShapes)
 	{
-		//TODO
 		return true;
 	}
 
-	//
-	//
-	//
 	virtual bool Initialize(TConstArrayView<NNE::FTensorDesc> Inputs, TConstArrayView<NNE::FTensorDesc> Outputs, const NNE::FAttributeMap& Attributes) override
 	{
 		check(Inputs.Num() == 1);
@@ -77,18 +64,12 @@ public:
 		return true;
 	}
 
-	//
-	//
-	//
 	virtual int PrepareOutputs(TConstArrayView<NNE::Internal::FTensorRef> InputTensors, TArrayView<NNE::Internal::FTensorRef> OutputTensors) const override
 	{
 		OutputTensors[0]->SetShape(InputTensors[0]->GetShape());
 		return 0;
 	}
 
-	//
-	//
-	//
 	virtual bool Create(IDMLDevice* Device, TConstArrayView<NNE::Internal::FTensorRef> InputTensors, TConstArrayView<NNE::Internal::FTensorRef> OutputTensors) override
 	{
 		const NNE::Internal::FTensor& InputTensor = *InputTensors[0];
diff --git a/Engine/Plugins/Experimental/NNERuntimeRDG/Source/NNERuntimeRDG/Private/Dml/NNEDmlOperatorPad.cpp b/Engine/Plugins/Experimental/NNERuntimeRDG/Source/NNERuntimeRDG/Private/Dml/NNEDmlOperatorPad.cpp
--- a/Engine/Plugins/Experimental/NNERuntimeRDG/Source/NNERuntimeRDG/Private/Dml/NNEDmlOperatorPad.cpp
+++ b/Engine/Plugins/Experimental/NNERuntimeRDG/Source/NNERuntimeRDG/Private/Dml/NNEDmlOperatorPad.cpp
@@ -5,18 +5,13 @@
 
 namespace UE::NNERuntimeRDG::Private::Dml
 {
-//
-//
-//
+
 class FOperatorDmlPad : public FOperatorDml
 {
 	float					Value;
 	DML_PADDING_MODE		Mode;
 	mutable TArray<int32>	Pads;
 
-	//
-	//
-	//
 	static DML_PADDING_MODE ModeFromString(FStringView StringVal)
 	{
 		if (FCString::Stricmp(StringVal.GetData(), TEXT("CONSTANT")) == 0)
@@ -39,26 +34,16 @@ class FOperatorDmlPad : public FOperatorDml
 
 public:
 
-	//
-	//
-	//
 	static FOperatorDml* Create()
 	{
 		return new FOperatorDmlPad();
 	}
 
-	//
-	//
-	//
 	static bool Validate(const NNE::FAttributeMap& AttributeMap, TConstArrayView<ENNETensorDataType> InputTypes, TConstArrayView<NNE::FSymbolicTensorShape> InputShapes)
 	{
-		//TODO
 		return true;
 	}
 
-	//
-	//
-	//
 	virtual bool Initialize(TConstArrayView<NNE::FTensorDesc> Inputs, TConstArrayView<NNE::FTensorDesc> Outputs, const NNE::FAttributeMap& Attributes) override
 	{
 		check(Inputs.Num() >= 1);
@@ -80,9 +65,6 @@ public:
 		return true;
 	}
 
-	//
-	//
-	//
 	virtual int PrepareOutputs(TConstArrayView<NNE::Internal::FTensorRef> InputTensors, TArrayView<NNE::Internal::FTensorRef> OutputTensors) const override
 	{
 		if (InputTensors.Num() >= 2)
@@ -134,9 +116,6 @@ public:
 		return 0;
 	}
 
-	//
-	//
-	//
 	virtual bool Create(IDMLDevice* Device, TConstArrayView<NNE::Internal::FTensorRef> InputTensors, TConstArrayView<NNE::Internal::FTensorRef> OutputTensors) override
 	{
 		const NNE::Internal::FTensor& InputTensor = *InputTensors[0];
diff --git a/Engine/Plugins/Experimental/NNERuntimeRDG/Source/NNERuntimeRDG/Private/Dml/NNEDmlOperatorPool.cpp b/Engine/Plugins/Experimental/NNERuntimeRDG/Source/NNERuntimeRDG/Private/Dml/NNEDmlOperatorPool.cpp
--- a/Engine/Plugins/Experimental/NNERuntimeRDG/Source/NNERuntimeRDG/Private/Dml/NNEDmlOperatorPool.cpp
+++ b/Engine/Plugins/Experimental/NNERuntimeRDG/Source/NNERuntimeRDG/Private/Dml/NNEDmlOperatorPool.cpp
@@ -18,18 +18,12 @@ class FOperatorDmlPool : public FOperatorDml
 {
 	using FIntArray = TArray<int32>;
 
-	//
-	//
-	//
 	class FPoolingArgs : public FKernelArgs
 	{
 		Util::FSmallUIntArray	KernelShape;
 
 	public:
 
-		//
-		//
-		//
 		bool Init(const NNE::FAttributeMap& Attributes, const int32 InputShapeRank, bool bInIsGlobalKernel)
 		{
 			if (!FKernelArgs::Init(Attributes, InputShapeRank, UseGlobalPooling, /*bIsTransposed*/ false))
@@ -62,9 +56,6 @@ class FOperatorDmlPool : public FOperatorDml
 			return true;
 		}
 
-		//
-		//
-		//
 		bool Evaluate(TConstArrayView<uint32> InputShape)
 		{
 			if (!bIsGlobalKernel)
@@ -79,9 +70,6 @@ class FOperatorDmlPool : public FOperatorDml
 			return true;
 		}
 
-		//
-		//
-		//
 		TConstArrayView<uint32> GetWindowSize() const
 		{
 			return WindowSize;
@@ -94,26 +82,16 @@ class FOperatorDmlPool : public FOperatorDml
 
 public:
 
-	//
-	//
-	//
 	static FOperatorDml* Create()
 	{
 		return new FOperatorDmlPool<TDmlPoolOpDesc, DmlOpType, UseGlobalPooling>();
 	}
 
-	//
-	//
-	//
 	static bool Validate(const NNE::FAttributeMap& AttributeMap, TConstArrayView<ENNETensorDataType> InputTypes, TConstArrayView<NNE::FSymbolicTensorShape> InputShapes)
 	{
-		//TODO
 		return true;
 	}
 
-	//
-	//
-	//
 	virtual bool Initialize(TConstArrayView<NNE::FTensorDesc> Inputs, TConstArrayView<NNE::FTensorDesc> Outputs, const NNE::FAttributeMap& Attributes) override
 	{
 		// TODO: int64 attributes
@@ -149,9 +127,6 @@ public:
 		return true;
 	}
 
-	//
-	//
-	//
 	virtual int PrepareOutputs(TConstArrayView<NNE::Internal::FTensorRef> InputTensors, TArrayView<NNE::Internal::FTensorRef> OutputTensors) const override
 	{
 		TConstArrayView<uint32> InputShape = InputTensors[0]->GetShape().GetData();
@@ -191,9 +166,6 @@ public:
 		return 0;
 	}
 
-	//
-	//
-	//
 	virtual bool Create(IDMLDevice* Device, TConstArrayView<NNE::Internal::FTensorRef> InputTensors, TConstArrayView<NNE::Internal::FTensorRef> OutputTensors) override
 	{
 		TConstArrayView<uint32> InputShape = InputTensors[0]->GetShape().GetData();
@@ -289,9 +261,6 @@ public:
 	}
 };
 
-//
-//
-//
 #define NNE_DML_REGISTER_POOLING_OP(OpName, DmlPrefix, UseGlobalPooling) \
 struct FDmlOperator##OpName##Registrator \
 { \
@@ -303,9 +272,7 @@ struct FDmlOperator##OpName##Registrator \
 \
 static FDmlOperator##OpName##Registrator RegisterDmlOperator##OpName;
 
-//
 // Register pooling operator on Module startup
-//
 NNE_DML_REGISTER_POOLING_OP(MaxPool, MAX, false)
 NNE_DML_REGISTER_POOLING_OP(GlobalMaxPool, MAX, true)
 NNE_DML_REGISTER_POOLING_OP(AveragePool, AVERAGE, false)
diff --git a/Engine/Plugins/Experimental/NNERuntimeRDG/Source/NNERuntimeRDG/Private/Dml/NNEDmlOperatorReduce.cpp b/Engine/Plugins/Experimental/NNERuntimeRDG/Source/NNERuntimeRDG/Private/Dml/NNEDmlOperatorReduce.cpp
--- a/Engine/Plugins/Experimental/NNERuntimeRDG/Source/NNERuntimeRDG/Private/Dml/NNEDmlOperatorReduce.cpp
+++ b/Engine/Plugins/Experimental/NNERuntimeRDG/Source/NNERuntimeRDG/Private/Dml/NNEDmlOperatorReduce.cpp
@@ -8,7 +8,6 @@
 namespace UE::NNERuntimeRDG::Private::Dml
 {
 
-//
 // Reduce operator covers:
 // ReduceL1
 // ReduceL2
@@ -23,9 +22,6 @@ namespace UE::NNERuntimeRDG::Private::Dml
 template<DML_REDUCE_FUNCTION ReduceFunc>
 class FOperatorDmlReduce : public FOperatorDml
 {
-	//
-	//
-	//
 	inline static void HandleEmptyAxes(TArray<int32>& Axes, int32 Rank)
 	{
 		if (Axes.IsEmpty())
@@ -47,26 +43,16 @@ class FOperatorDmlReduce : public FOperatorDml
 
 public:
 
-	//
-	//
-	//
 	static FOperatorDml* Create()
 	{
 		return new FOperatorDmlReduce<ReduceFunc>();
 	}
 
-	//
-	//
-	//
 	static bool Validate(const NNE::FAttributeMap& AttributeMap, TConstArrayView<ENNETensorDataType> InputTypes, TConstArrayView<NNE::FSymbolicTensorShape> InputShapes)
 	{
-		//TODO
 		return true;
 	}
 
-	//
-	//
-	//
 	virtual bool Initialize(TConstArrayView<NNE::FTensorDesc> Inputs, TConstArrayView<NNE::FTensorDesc> Outputs, const NNE::FAttributeMap& Attributes) override
 	{
 		checkf(Inputs.Num() == 1, TEXT("Dml Reduce op supports only 1 input"));
@@ -108,9 +94,6 @@ public:
 		return true;
 	}
 
-	//
-	//
-	//
 	virtual int PrepareOutputs(TConstArrayView<NNE::Internal::FTensorRef> InputTensors, TArrayView<NNE::Internal::FTensorRef> OutputTensors) const override
 	{
 		TConstArrayView<uint32> InputShape = InputTensors[0]->GetShape().GetData();
@@ -151,9 +134,6 @@ public:
 		return 0;
 	}
 
-	//
-	//
-	//
 	virtual bool Create(IDMLDevice* Device, TConstArrayView<NNE::Internal::FTensorRef> InputTensors, TConstArrayView<NNE::Internal::FTensorRef> OutputTensors) override
 	{
 		FTensorDescDml DmlInputTensorDesc;
diff --git a/Engine/Plugins/Experimental/NNERuntimeRDG/Source/NNERuntimeRDG/Private/Dml/NNEDmlOperatorResample.cpp b/Engine/Plugins/Experimental/NNERuntimeRDG/Source/NNERuntimeRDG/Private/Dml/NNEDmlOperatorResample.cpp
--- a/Engine/Plugins/Experimental/NNERuntimeRDG/Source/NNERuntimeRDG/Private/Dml/NNEDmlOperatorResample.cpp
+++ b/Engine/Plugins/Experimental/NNERuntimeRDG/Source/NNERuntimeRDG/Private/Dml/NNEDmlOperatorResample.cpp
@@ -32,9 +32,6 @@ void RemoveValuesByIndex(TConstArrayView<uint32> Indices, TArray<TData, TAllocat
 template <bool IsResize>
 class FOperatorDmlResample : public FOperatorDml
 {
-	//
-	//
-	//
 	enum ECoordTransformMode : uint8
 	{
 		None,
@@ -48,9 +45,6 @@ class FOperatorDmlResample : public FOperatorDml
 	DML_INTERPOLATION_MODE	Mode;
 	ECoordTransformMode		CoordTransformMode;
 
-	//
-	//
-	//
 	static DML_INTERPOLATION_MODE ModeFromString(FStringView StringVal)
 	{
 		if (FCString::Stricmp(StringVal.GetData(), TEXT("NEAREST")) == 0)
@@ -70,26 +64,16 @@ class FOperatorDmlResample : public FOperatorDml
 
 public:
 
-	//
-	//
-	//
 	static FOperatorDml* Create()
 	{
 		return new FOperatorDmlResample<IsResize>();
 	}
 
-	//
-	//
-	//
 	static bool Validate(const NNE::FAttributeMap& AttributeMap, TConstArrayView<ENNETensorDataType> InputTypes, TConstArrayView<NNE::FSymbolicTensorShape> InputShapes)
 	{
-		//TODO
 		return true;
 	}
 
-	//
-	//
-	//
 	virtual bool Initialize(TConstArrayView<NNE::FTensorDesc> Inputs, TConstArrayView<NNE::FTensorDesc> Outputs, const NNE::FAttributeMap& Attributes) override
 	{
 		check(Inputs.Num() >= 1 && Inputs.Num() < 4);
@@ -144,7 +128,6 @@ public:
 			}
 		}
 
-		// Notify DML graph that we have constant CPU tensor inputs
 		for (int i = 1; i < Inputs.Num(); ++i)
 		{
 			ConstantCPUInputs.Add(i);
@@ -153,9 +136,6 @@ public:
 		return true;
 	}
 
-	//
-	//
-	//
 	virtual int PrepareOutputs(TConstArrayView<NNE::Internal::FTensorRef> InputTensors, TArrayView<NNE::Internal::FTensorRef> OutputTensors) const override
 	{
 		const NNE::Internal::FTensor& InputTensor = *InputTensors[0];
@@ -225,9 +205,6 @@ public:
 		return 0;
 	}
 
-	//
-	//
-	//
 	virtual bool Create(IDMLDevice* Device, TConstArrayView<NNE::Internal::FTensorRef> InputTensors, TConstArrayView<NNE::Internal::FTensorRef> OutputTensors) override
 	{
 		const NNE::Internal::FTensor& InputTensor = *InputTensors[0];
@@ -282,7 +259,6 @@ public:
 			}
 		}
 
-		// Initialize tensor descriptors
 		FTensorDescDml	DmlInputTensorDesc;
 		FTensorDescDml	DmlOutputTensorDesc;
 
diff --git a/Engine/Plugins/Experimental/NNERuntimeRDG/Source/NNERuntimeRDG/Private/Dml/NNEDmlOperatorReshape.cpp b/Engine/Plugins/Experimental/NNERuntimeRDG/Source/NNERuntimeRDG/Private/Dml/NNEDmlOperatorReshape.cpp
--- a/Engine/Plugins/Experimental/NNERuntimeRDG/Source/NNERuntimeRDG/Private/Dml/NNEDmlOperatorReshape.cpp
+++ b/Engine/Plugins/Experimental/NNERuntimeRDG/Source/NNERuntimeRDG/Private/Dml/NNEDmlOperatorReshape.cpp
@@ -11,9 +11,6 @@
 namespace UE::NNERuntimeRDG::Private::Dml
 {
 
-/**
- * Reshape
- */
 class FOperatorDmlReshape : public FOperatorDml
 {
 	mutable Util::FSmallUIntArray	OutputShape;
@@ -21,26 +18,16 @@ class FOperatorDmlReshape : public FOperatorDml
 
 public:
 
-	//
-	//
-	//
 	static FOperatorDml* Create()
 	{
 		return new FOperatorDmlReshape();
 	}
 
-	//
-	//
-	//
     static bool Validate(const NNE::FAttributeMap& AttributeMap, TConstArrayView<ENNETensorDataType> InputTypes, TConstArrayView<NNE::FSymbolicTensorShape> InputShapes)
 	{
-		//TODO
 		return true;
 	}
 
-	//
-	//
-	//
 	virtual bool Initialize(TConstArrayView<NNE::FTensorDesc> Inputs, TConstArrayView<NNE::FTensorDesc> Outputs, const NNE::FAttributeMap& Attributes) override
 	{
 		check(Inputs.Num() == 2);
@@ -53,9 +40,6 @@ public:
 		return true;
 	}
 
-	//
-	//
-	//
 	virtual int PrepareOutputs(TConstArrayView<NNE::Internal::FTensorRef> InputTensors, TArrayView<NNE::Internal::FTensorRef> OutputTensors) const override
 	{
 		// Shape tensor must be constant!
@@ -96,9 +80,6 @@ public:
 		return 0;
 	}
 
-	//
-	//
-	//
 	virtual bool Create(IDMLDevice* Device, TConstArrayView<NNE::Internal::FTensorRef> InputTensors, TConstArrayView<NNE::Internal::FTensorRef> OutputTensors) override
 	{
         FTensorDescDml DmlInputTensorDesc;
diff --git a/Engine/Plugins/Experimental/NNERuntimeRDG/Source/NNERuntimeRDG/Private/Dml/NNEDmlOperatorSlice.cpp b/Engine/Plugins/Experimental/NNERuntimeRDG/Source/NNERuntimeRDG/Private/Dml/NNEDmlOperatorSlice.cpp
--- a/Engine/Plugins/Experimental/NNERuntimeRDG/Source/NNERuntimeRDG/Private/Dml/NNEDmlOperatorSlice.cpp
+++ b/Engine/Plugins/Experimental/NNERuntimeRDG/Source/NNERuntimeRDG/Private/Dml/NNEDmlOperatorSlice.cpp
@@ -33,9 +33,6 @@ void Iota(ArrayT& Array, ValueT Value)
 
 } // Util
 
-/**
- * Slice
- */
 class FOperatorDmlSlice : public FOperatorDml
 {	
 	template<typename DataType>
@@ -105,6 +102,7 @@ class FOperatorDmlSlice : public FOperatorDml
 			{
 				Start += (DataType) Dim;
 			}
+			
 			if (End < 0 && Start > TNumericLimits<DataType>::Min())
 			{
 				End += (DataType) Dim;
diff --git a/Engine/Plugins/Experimental/NNERuntimeRDG/Source/NNERuntimeRDG/Private/Dml/NNEDmlOperatorSpaceToDepth.cpp b/Engine/Plugins/Experimental/NNERuntimeRDG/Source/NNERuntimeRDG/Private/Dml/NNEDmlOperatorSpaceToDepth.cpp
--- a/Engine/Plugins/Experimental/NNERuntimeRDG/Source/NNERuntimeRDG/Private/Dml/NNEDmlOperatorSpaceToDepth.cpp
+++ b/Engine/Plugins/Experimental/NNERuntimeRDG/Source/NNERuntimeRDG/Private/Dml/NNEDmlOperatorSpaceToDepth.cpp
@@ -8,9 +8,6 @@
 namespace UE::NNERuntimeRDG::Private::Dml
 {
 
-//
-//
-//
 class FOperatorDmlSpaceToDepth : public FOperatorDml
 {
 	enum InputDims
@@ -23,26 +20,16 @@ class FOperatorDmlSpaceToDepth : public FOperatorDml
 
 public:
 
-	//
-	//
-	//
 	static FOperatorDml* Create()
 	{
 		return new FOperatorDmlSpaceToDepth();
 	}
 
-	//
-	//
-	//
 	static bool Validate(const NNE::FAttributeMap& AttributeMap, TConstArrayView<ENNETensorDataType> InputTypes, TConstArrayView<NNE::FSymbolicTensorShape> InputShapes)
 	{
-		//TODO
 		return true;
 	}
 
-	//
-	//
-	//
 	virtual bool Initialize(TConstArrayView<NNE::FTensorDesc> Inputs, TConstArrayView<NNE::FTensorDesc> Outputs, const NNE::FAttributeMap& Attributes) override
 	{
 		check(Inputs.Num() == 1);
@@ -62,9 +49,6 @@ public:
 		return true;
 	}
 
-	//
-	//
-	//
 	virtual int PrepareOutputs(TConstArrayView<NNE::Internal::FTensorRef> InputTensors, TArrayView<NNE::Internal::FTensorRef> OutputTensors) const override
 	{
 		TConstArrayView<uint32>		InputShape = InputTensors[0]->GetShape().GetData();
@@ -81,9 +65,6 @@ public:
 		return 0;
 	}
 
-	//
-	//
-	//
 	virtual bool Create(IDMLDevice* Device, TConstArrayView<NNE::Internal::FTensorRef> InputTensors, TConstArrayView<NNE::Internal::FTensorRef> OutputTensors) override
 	{
 		const NNE::Internal::FTensor& InputTensor = *InputTensors[0];
diff --git a/Engine/Plugins/Experimental/NNERuntimeRDG/Source/NNERuntimeRDG/Private/Dml/NNEDmlOperatorSplit.cpp b/Engine/Plugins/Experimental/NNERuntimeRDG/Source/NNERuntimeRDG/Private/Dml/NNEDmlOperatorSplit.cpp
--- a/Engine/Plugins/Experimental/NNERuntimeRDG/Source/NNERuntimeRDG/Private/Dml/NNEDmlOperatorSplit.cpp
+++ b/Engine/Plugins/Experimental/NNERuntimeRDG/Source/NNERuntimeRDG/Private/Dml/NNEDmlOperatorSplit.cpp
@@ -6,9 +6,7 @@
 
 namespace UE::NNERuntimeRDG::Private::Dml
 {
-//
-//
-//
+
 class FOperatorDmlSplit : public FOperatorDml
 {
 	TArray<int32>	Split;
@@ -16,26 +14,16 @@ class FOperatorDmlSplit : public FOperatorDml
 
 public:
 
-	//
-	//
-	//
 	static FOperatorDml* Create()
 	{
 		return new FOperatorDmlSplit();
 	}
 
-	//
-	//
-	//
 	static bool Validate(const NNE::FAttributeMap& AttributeMap, TConstArrayView<ENNETensorDataType> InputTypes, TConstArrayView<NNE::FSymbolicTensorShape> InputShapes)
 	{
-		//TODO
 		return true;
 	}
 
-	//
-	//
-	//
 	virtual bool Initialize(TConstArrayView<NNE::FTensorDesc> Inputs, TConstArrayView<NNE::FTensorDesc> Outputs, const NNE::FAttributeMap& Attributes) override
 	{
 		check(Inputs.Num() == 1);
@@ -63,7 +51,6 @@ public:
 		}
 
 		// Check split size is correct
-		//TODO: code should be validated by a validator which should report to the user if the contract is broken. We do it here instead.
 		for (int Idx = 0; Idx < Outputs.Num(); ++Idx)
 		{
 			if (Outputs[Idx].GetShape().Rank() != Inputs[0].GetShape().Rank())
@@ -76,9 +63,6 @@ public:
 		return true;
 	}
 
-	//
-	//
-	//
 	virtual int PrepareOutputs(TConstArrayView<NNE::Internal::FTensorRef> InputTensors, TArrayView<NNE::Internal::FTensorRef> OutputTensors) const override
 	{
 		const NNE::Internal::FTensor&	InputTensor = *InputTensors[0];
@@ -127,9 +111,6 @@ public:
 		return 0;
 	}
 
-	//
-	//
-	//
 	virtual bool Create(IDMLDevice* Device, TConstArrayView<NNE::Internal::FTensorRef> InputTensors, TConstArrayView<NNE::Internal::FTensorRef> OutputTensors) override
 	{
 		FTensorDescDml InputTensorDesc;
diff --git a/Engine/Plugins/Experimental/NNERuntimeRDG/Source/NNERuntimeRDG/Private/Dml/NNEDmlOperatorSqueeze.cpp b/Engine/Plugins/Experimental/NNERuntimeRDG/Source/NNERuntimeRDG/Private/Dml/NNEDmlOperatorSqueeze.cpp
--- a/Engine/Plugins/Experimental/NNERuntimeRDG/Source/NNERuntimeRDG/Private/Dml/NNEDmlOperatorSqueeze.cpp
+++ b/Engine/Plugins/Experimental/NNERuntimeRDG/Source/NNERuntimeRDG/Private/Dml/NNEDmlOperatorSqueeze.cpp
@@ -8,35 +8,22 @@
 namespace UE::NNERuntimeRDG::Private::Dml
 {
 
-//
-//
-//
 class FOperatorDmlSqueeze : public FOperatorDml
 {
 	mutable TArray<int32>	Axes;
 
 public:
 
-	//
-	//
-	//
 	static FOperatorDml* Create()
 	{
 		return new FOperatorDmlSqueeze();
 	}
 
-	//
-	//
-	//
 	static bool Validate(const NNE::FAttributeMap& AttributeMap, TConstArrayView<ENNETensorDataType> InputTypes, TConstArrayView<NNE::FSymbolicTensorShape> InputShapes)
 	{
-		//TODO
 		return true;
 	}
 
-	//
-	//
-	//
 	virtual bool Initialize(TConstArrayView<NNE::FTensorDesc> Inputs, TConstArrayView<NNE::FTensorDesc> Outputs, const NNE::FAttributeMap& Attributes) override
 	{
 		check(Inputs.Num() == 1);
@@ -55,9 +42,6 @@ public:
 		return true;
 	}
 
-	//
-	//
-	//
 	virtual int PrepareOutputs(TConstArrayView<NNE::Internal::FTensorRef> InputTensors, TArrayView<NNE::Internal::FTensorRef> OutputTensors) const override
 	{
 		TConstArrayView<uint32>	InputShape = InputTensors[0]->GetShape().GetData();
@@ -103,9 +87,6 @@ public:
 		return 0;
 	}
 
-	//
-	//
-	//
 	virtual bool Create(IDMLDevice* Device, TConstArrayView<NNE::Internal::FTensorRef> InputTensors, TConstArrayView<NNE::Internal::FTensorRef> OutputTensors) override
 	{
 		FTensorDescDml DmlTensorDesc;
diff --git a/Engine/Plugins/Experimental/NNERuntimeRDG/Source/NNERuntimeRDG/Private/Dml/NNEDmlOperatorTranspose.cpp b/Engine/Plugins/Experimental/NNERuntimeRDG/Source/NNERuntimeRDG/Private/Dml/NNEDmlOperatorTranspose.cpp
--- a/Engine/Plugins/Experimental/NNERuntimeRDG/Source/NNERuntimeRDG/Private/Dml/NNEDmlOperatorTranspose.cpp
+++ b/Engine/Plugins/Experimental/NNERuntimeRDG/Source/NNERuntimeRDG/Private/Dml/NNEDmlOperatorTranspose.cpp
@@ -10,14 +10,9 @@
 namespace UE::NNERuntimeRDG::Private::Dml
 {
 
-/**
- * Transpose
- */
 class FOperatorDmlTranspose : public FOperatorDml
 {
-	//
 	// Apply permutations to input array view
-	//
 	Util::FSmallUIntArray Permute(TConstArrayView<uint32> InputView) const
 	{
 		Util::FSmallUIntArray Permuted;
@@ -34,26 +29,16 @@ class FOperatorDmlTranspose : public FOperatorDml
 
 public:
 
-	//
-	//
-	//
 	static FOperatorDml* Create()
 	{
 		return new FOperatorDmlTranspose();
 	}
 
-	//
-	//
-	//
 	static bool Validate(const NNE::FAttributeMap& AttributeMap, TConstArrayView<ENNETensorDataType> InputTypes, TConstArrayView<NNE::FSymbolicTensorShape> InputShapes)
 	{
-		//TODO
 		return true;
 	}
 
-	//
-	//
-	//
 	virtual bool Initialize(TConstArrayView<NNE::FTensorDesc> Inputs, TConstArrayView<NNE::FTensorDesc> Outputs, const NNE::FAttributeMap& Attributes) override
 	{
 		check(Inputs.Num() == 1);
@@ -77,9 +62,6 @@ public:
 
 	}
 
-	//
-	//
-	//
 	virtual int PrepareOutputs(TConstArrayView<NNE::Internal::FTensorRef> InputTensors, TArrayView<NNE::Internal::FTensorRef> OutputTensors) const override
 	{
 		Util::FSmallUIntArray OutputShape = Permute(InputTensors[0]->GetShape().GetData());
@@ -89,9 +71,6 @@ public:
 		return 0;
 	}
 
-	//
-	//
-	//
 	virtual bool Create(IDMLDevice* Device, TConstArrayView<NNE::Internal::FTensorRef> InputTensors, TConstArrayView<NNE::Internal::FTensorRef> OutputTensors) override
 	{
 		TConstArrayView<uint32>	InputShape = InputTensors[0]->GetShape().GetData();
diff --git a/Engine/Plugins/Experimental/NNERuntimeRDG/Source/NNERuntimeRDG/Private/Dml/NNEDmlOperatorUnsqueeze.cpp b/Engine/Plugins/Experimental/NNERuntimeRDG/Source/NNERuntimeRDG/Private/Dml/NNEDmlOperatorUnsqueeze.cpp
--- a/Engine/Plugins/Experimental/NNERuntimeRDG/Source/NNERuntimeRDG/Private/Dml/NNEDmlOperatorUnsqueeze.cpp
+++ b/Engine/Plugins/Experimental/NNERuntimeRDG/Source/NNERuntimeRDG/Private/Dml/NNEDmlOperatorUnsqueeze.cpp
@@ -7,35 +7,23 @@
 
 namespace UE::NNERuntimeRDG::Private::Dml
 {
-//
-//
-//
+
 class FOperatorDmlUnsqueeze : public FOperatorDml
 {
 	TArray<int32> Axes;
 	
 public:
 
-	//
-	//
-	//
 	static FOperatorDml* Create()
 	{
 		return new FOperatorDmlUnsqueeze();
 	}
 
-	//
-	//
-	//
 	static bool Validate(const NNE::FAttributeMap& AttributeMap, TConstArrayView<ENNETensorDataType> InputTypes, TConstArrayView<NNE::FSymbolicTensorShape> InputShapes)
 	{
-		//TODO
 		return true;
 	}
 
-	//
-	//
-	//
 	virtual bool Initialize(TConstArrayView<NNE::FTensorDesc> Inputs, TConstArrayView<NNE::FTensorDesc> Outputs, const NNE::FAttributeMap& Attributes) override
 	{
 		check(Inputs.Num() == 1);
@@ -58,13 +46,10 @@ public:
 		return true;
 	}
 
-	//
-	//
-	//
 	virtual int PrepareOutputs(TConstArrayView<NNE::Internal::FTensorRef> InputTensors, TArrayView<NNE::Internal::FTensorRef> OutputTensors) const override
 	{
 		TConstArrayView<uint32>	InputShape = InputTensors[0]->GetShape().GetData();
-		const int32_t			OutShapeRank = InputShape.Num() + Axes.Num();
+		const int32				OutShapeRank = InputShape.Num() + Axes.Num();
 
 		Util::FSmallUIntArray	OutputShape;
 		OutputShape.Reserve(OutShapeRank);
@@ -86,9 +71,6 @@ public:
 		return 0;
 	}
 
-	//
-	//
-	//
 	virtual bool Create(IDMLDevice* Device, TConstArrayView<NNE::Internal::FTensorRef> InputTensors, TConstArrayView<NNE::Internal::FTensorRef> OutputTensors) override
 	{
 		const NNE::Internal::FTensor& InputTensor = *InputTensors[0];
diff --git a/Engine/Plugins/Experimental/NNERuntimeRDG/Source/NNERuntimeRDG/Private/Dml/NNEDmlOperatorUtils.h b/Engine/Plugins/Experimental/NNERuntimeRDG/Source/NNERuntimeRDG/Private/Dml/NNEDmlOperatorUtils.h
--- a/Engine/Plugins/Experimental/NNERuntimeRDG/Source/NNERuntimeRDG/Private/Dml/NNEDmlOperatorUtils.h
+++ b/Engine/Plugins/Experimental/NNERuntimeRDG/Source/NNERuntimeRDG/Private/Dml/NNEDmlOperatorUtils.h
@@ -13,9 +13,6 @@
 namespace UE::NNERuntimeRDG::Private::Dml
 {
 
-//
-//
-//
 enum EAutoPad
 {
 	NOTSET,
@@ -24,9 +21,6 @@ enum EAutoPad
 	VALID
 };
 
-//
-//
-//
 static EAutoPad AutoPadFromString(FStringView StringVal) 
 {
 	if (FCString::Stricmp(StringVal.GetData(), TEXT("NOTSET")) == 0) 
@@ -51,9 +45,6 @@ static EAutoPad AutoPadFromString(FStringView StringVal)
 	}
 }
 
-//
-//
-//
 template<typename T>
 static bool IsEqualOrBroadcastable(TConstArrayView<T> ShapeA, TConstArrayView<T> ShapeB)
 {
@@ -64,7 +55,7 @@ static bool IsEqualOrBroadcastable(TConstArrayView<T> ShapeA, TConstArrayView<T>
 
 	int32 BIdx = 0;
 
-	for(int32 Idx = 0; Idx < ShapeA.Num() && BIdx < ShapeB.Num(); ++Idx)
+	for (int32 Idx = 0; Idx < ShapeA.Num() && BIdx < ShapeB.Num(); ++Idx)
 	{
 		if (BIdx != 0)
 		{
@@ -77,6 +68,7 @@ static bool IsEqualOrBroadcastable(TConstArrayView<T> ShapeA, TConstArrayView<T>
 			}
 			++BIdx;
 		}
+
 		if (BIdx == 0 && ShapeA[Idx] == ShapeB[BIdx])
 		{
 			++BIdx;
@@ -91,9 +83,6 @@ static bool IsEqualOrBroadcastable(TConstArrayView<T> ShapeA, TConstArrayView<T>
 	return true;
 }
 
-//
-//
-//
 static bool CheckGenericTensor(ENNETensorDataType DataType, const NNE::FSymbolicTensorShape& TensorShape)
 {
 	const int32 MinTensorRank(0), MaxTensorRank(DML_TENSOR_DIMENSION_COUNT_MAX1);
@@ -107,9 +96,6 @@ static bool CheckGenericTensor(ENNETensorDataType DataType, const NNE::FSymbolic
 	return true;
 }
 
-//
-//
-//
 static bool CheckElementwiseTensor(ENNETensorDataType DataType, const NNE::FSymbolicTensorShape& TensorShape)
 {
 	if (DataType != ENNETensorDataType::Float)
@@ -126,9 +112,6 @@ static bool CheckElementwiseTensor(ENNETensorDataType DataType, const NNE::FSymb
 	return true;
 }
 
-//
-//
-//
 static Util::FSmallUIntArray KernelPadding(
 	TConstArrayView<uint32> InputShape, TConstArrayView<uint32> WindowSize, 
 	TConstArrayView<uint32> Dilations, TConstArrayView<uint32> Strides
@@ -153,9 +136,9 @@ static Util::FSmallUIntArray KernelPadding(
 	return Padding;
 }
 
-//
-// Used by FKernelArgs
-//
+/**
+* Utility class used by FKernelArgs 
+*/
 class FPaddingsHelper
 {
 	Util::FSmallUIntArray	Pads;
@@ -164,9 +147,6 @@ class FPaddingsHelper
 
 public:
 
-	//
-	//
-	//
 	bool Init(const NNE::FAttributeMap& Attributes, uint32 InputRank)
 	{
 		NumSpatialDimensions = InputRank - NonspatialDimensionCount;
@@ -196,9 +176,6 @@ public:
 		return true;
 	}
 
-	//
-	//
-	//
 	bool Evaluate(Util::FSmallUIntArray& OutStartPadding, Util::FSmallUIntArray& OutEndPadding, TConstArrayView<uint32> Padding)
 	{
 		if (AutoPad == EAutoPad::NOTSET)
@@ -238,57 +215,40 @@ public:
 	}
 };
 
-//
-// This is a base class that is used for Conv, ConvTranspose, Pool (both local and global) operators and MaxUnpool
-//
+/** 
+* This is a base class that is used for Conv, ConvTranspose, Pool(both local and global) operators and MaxUnpool 
+* NOTE: WindowSize needs to be set recomputed in the sub-class
+*/
 class FKernelArgs
 {
 
 public:
 
-	//
-	//
-	//
 	uint32 GetNumDimensions() const
 	{
 		return NumDimensions;
 	}
 
-	//
-	//
-	//
 	TConstArrayView<uint32> GetStrides() const
 	{
 		return Strides;
 	}
 
-	//
-	//
-	//
 	TConstArrayView<uint32> GetDilations() const
 	{
 		return Dilations;
 	}
 
-	//
-	//
-	//
 	TConstArrayView<uint32> GetStartPadding()
 	{
 		return StartPadding;
 	}
 
-	//
-	//
-	//
 	TConstArrayView<uint32> GetEndPadding()
 	{
 		return EndPadding;
 	}
 
-	//
-	//
-	//
 	TConstArrayView<uint32> GetOutputShape() const
 	{
 		return OutputShape;
@@ -315,9 +275,6 @@ protected:
 
 protected:
 
-	//
-	//
-	//
 	bool Init(const NNE::FAttributeMap& Attributes, int32 InputShapeRank, bool bInIsGlobalKernel, bool bInIsTransposed)
 	{
 		check(InputShapeRank > NonspatialDimensionCount);
@@ -333,8 +290,6 @@ protected:
 			StartPadding.Init(0, NumDimensions);
 			EndPadding.Init(0, NumDimensions);
 			OutPadding.Init(0, NumDimensions);
-
-			// NOTE: WindowSize needs to be set recomputed in the sub-class
 		}
 		else
 		{
@@ -410,9 +365,6 @@ protected:
 
 public:
 
-	//
-	//
-	//
 	void Evaluate(TConstArrayView<uint32> InputShape, TConstArrayView<uint32> PaddingsValue = MakeEmptyConstArrayView<uint32>())
 	{
 		if (bIsGlobalKernel && WindowSize.IsEmpty())
@@ -476,10 +428,6 @@ public:
 	}
 };
 
-
-//
-//
-//
 inline int32 HandleNegativeAxis(int32 Axis, int32 Rank)
 {
 	if (Axis < 0)
@@ -491,9 +439,6 @@ inline int32 HandleNegativeAxis(int32 Axis, int32 Rank)
 	return Axis;
 }
 
-//
-//
-//
 inline void HandleNegativeAxes(TArrayView<int32> Axes, int32 Rank)
 {
 	for (int32& Axis : Axes)
@@ -502,9 +447,6 @@ inline void HandleNegativeAxes(TArrayView<int32> Axes, int32 Rank)
 	}
 }
 
-//
-//
-//
 inline int32 GetDmlAxis(int32 OnnxAxis, int32 OnnxDim, int32 DmlDim)
 {
 	check(DmlDim >= OnnxDim);
@@ -514,9 +456,6 @@ inline int32 GetDmlAxis(int32 OnnxAxis, int32 OnnxDim, int32 DmlDim)
 	return DmlAxis;
 }
 
-//
-//
-//
 inline void SetDmlAxesFromOnnx(Util::FSmallUIntArray& DmlAxes, int32 Rank, TConstArrayView<int32> OnnxAxes)
 {
 	DmlAxes.Reset();
diff --git a/Engine/Plugins/Experimental/NNERuntimeRDG/Source/NNERuntimeRDG/Private/NNERuntimeRDGModel.cpp b/Engine/Plugins/Experimental/NNERuntimeRDG/Source/NNERuntimeRDG/Private/NNERuntimeRDGModel.cpp
--- a/Engine/Plugins/Experimental/NNERuntimeRDG/Source/NNERuntimeRDG/Private/NNERuntimeRDGModel.cpp
+++ b/Engine/Plugins/Experimental/NNERuntimeRDG/Source/NNERuntimeRDG/Private/NNERuntimeRDGModel.cpp
@@ -187,8 +187,6 @@ int32 FModelInstanceRDG::SetInputTensorShapes(TConstArrayView<NNE::FTensorShape>
 
 FRDGBufferDesc CreateRDGBufferDescForTensorRDG(const FTensorRDG& Tensor)
 {
-	// FIXME: CreateStructuredDesc() creates a crash on VulkanRHI
-	//FRDGBufferDesc Desc = FRDGBufferDesc::CreateStructuredDesc(Tensor.GetElementByteSize(), Tensor.GetVolume());
 	FRDGBufferDesc Desc = FRDGBufferDesc::CreateBufferDesc(Tensor.GetElementByteSize(), Tensor.GetVolume());
 
 	return Desc;
@@ -271,4 +269,4 @@ int32 FModelInstanceRDG::SetTensors(FRDGBuilder& GraphBuilder, FTensorRDGArray&
 	return -1;
 }
 
-}
\ No newline at end of file
+} // namespace UE::NNERuntimeRDG::Private
