[NNE] RDG Dml Skip empty tensors

Problem:
Resize operator requires non-positional argument, marked in onnx format as empty tensor. This wasn't supported in RDG Dml.

Solution:
Adapt setup code to allow for holes in the tensor numbering due to skipped empty tensors.
In practice turn AllSymbolicTensorDescs and AllTensorRDGRefs (FModelInstanceRDG) from arrays into maps (from tensor index to content) and adapt usages.

Tests:
- NNEIncubatorEditor Win64
- NNEIncubator Wn64
- Non-Unity compilation test (p4v)

#rnx

diff --git a/Engine/Plugins/Experimental/NNERuntimeRDG/Source/NNERuntimeRDG/Private/Dml/NNEDmlModel.cpp b/Engine/Plugins/Experimental/NNERuntimeRDG/Source/NNERuntimeRDG/Private/Dml/NNEDmlModel.cpp
--- a/Engine/Plugins/Experimental/NNERuntimeRDG/Source/NNERuntimeRDG/Private/Dml/NNEDmlModel.cpp
+++ b/Engine/Plugins/Experimental/NNERuntimeRDG/Source/NNERuntimeRDG/Private/Dml/NNEDmlModel.cpp
@@ -35,6 +35,9 @@ DECLARE_GPU_STAT_NAMED(GPU_STAT_DispatchTime, TEXT("NNEDmlModelDispatchTime"));
 DECLARE_GPU_STAT_NAMED(GPU_STAT_DispatchD3DTime, TEXT("NNEDmlModelDispatchD3DTime"));
 #endif
 
+// Empty tensor is marked by -1
+static constexpr int32 GEmptyTensorIdx = -1;
+
 FModelInfo* FModelInfo::Get()
 {
 	static FModelInfo Inst;
@@ -48,26 +51,31 @@ FModelInfo::FModelInfo()
 {
 }
 
+
 FGuid FModelInfo::GetGuid()
 {
 	return Guid;
 }
 
+
 int32 FModelInfo::GetVersion()
 {
 	return Version;
 }
 
+
 int32 FModelInfo::GetGuidSize()
 {
 	return sizeof(Guid);
 }
 
+
 int32 FModelInfo::GetVersionSize()
 {
 	return sizeof(Version);
 }
 
+
 bool FModelInfo::ValidateGuidAndVersion(const uint8* InGuid, const uint8* InVersion)
 {
 	bool bResult;
@@ -500,13 +508,16 @@ private:
 		NumInputs = 0;
 		NumOutputs = 0;
 
-		TensorInConnCounts.SetNumZeroed(InModel->AllSymbolicTensorDescs.Num());
-		TensorOutConnCounts.SetNumZeroed(InModel->AllSymbolicTensorDescs.Num());
+		TensorInConnCounts.SetNumZeroed(InModel->TensorIdxSpan);
+		TensorOutConnCounts.SetNumZeroed(InModel->TensorIdxSpan);
 		Operators.Reset(InModel->GraphOperators.Num());
 
 		for (int32 Idx : InModel->GraphOpInputIndices)
 		{
-			TensorInConnCounts[Idx] += 1;
+			if (Idx != GEmptyTensorIdx)
+			{
+				TensorInConnCounts[Idx] += 1;
+			}
 		}
 
 		for (int32 Idx : InModel->GraphOpOutputIndices)
@@ -671,6 +682,11 @@ private:
 			{
 				const int32 TensorIdx = InModel->GraphOpInputIndices[Idx + CurrOp.InputStart];
 
+				if (TensorIdx == GEmptyTensorIdx)
+				{
+					continue;
+				}
+
 				// Filter out the constant CPU inputs from the graph node inputs
 				if (InModel->ConstantCPUTensorIndices.Find(TensorIdx) != INDEX_NONE)
 				{
@@ -868,8 +884,16 @@ bool FModelInstance::Init(TConstArrayView<uint8> ModelData, FDmlDeviceContext* I
 				MemSizeWeights += Align(TensorData.Num(), DML_MINIMUM_BUFFER_TENSOR_ALIGNMENT);
 			}
 
-			OpInputTensors.Emplace(AllSymbolicTensorDescs[InputTensorIndex]);
-			GraphOpInputIndices.Emplace(InputTensorIndex);
+			if (Format.Tensors[InputTensorIndex].Type != ENNEFormatTensorType::Empty)
+			{
+				OpInputTensors.Emplace(AllSymbolicTensorDescs[InputTensorIndex]);
+				GraphOpInputIndices.Emplace(InputTensorIndex);
+			}
+			else
+			{
+				OpInputTensors.Emplace(NNE::FTensorDesc::Make(TEXT(""), {}, ENNETensorDataType::None));
+				GraphOpInputIndices.Emplace(GEmptyTensorIdx);
+			}
 		}
 
 		for (int32 OutputTensorIndex : Format.Operators[Idx].OutTensors)
@@ -916,7 +940,10 @@ bool FModelInstance::Init(TConstArrayView<uint8> ModelData, FDmlDeviceContext* I
 			const int32 ConstIdx = OpConstantCPUInputs[InputIdx];
 			const int32 TensorIdx = GraphOpInputIndices[OpDesc.InputStart + ConstIdx];
 
-			ConstantCPUTensorIndices.AddUnique(TensorIdx);
+			if (TensorIdx != GEmptyTensorIdx)
+			{
+				ConstantCPUTensorIndices.AddUnique(TensorIdx);
+			}
 		}
 	}
 
@@ -1298,7 +1325,11 @@ int FModelInstance::PrepareTensorShapesAndData()
 		for (int32 Idx = 0; Idx < OpDesc.InputCount; ++Idx)
 		{
 			const int32 InputIdx = GraphOpInputIndices[OpDesc.InputStart + Idx];
-			OpInputs.Emplace(AllTensorRDGRefs[InputIdx]);
+
+			if (InputIdx != GEmptyTensorIdx)
+			{
+				OpInputs.Emplace(AllTensorRDGRefs[InputIdx]);
+			}
 		}
 
 		for (int32 Idx = 0; Idx < OpDesc.OutputCount; ++Idx)
diff --git a/Engine/Plugins/Experimental/NNERuntimeRDG/Source/NNERuntimeRDG/Private/NNERuntimeRDGBase.h b/Engine/Plugins/Experimental/NNERuntimeRDG/Source/NNERuntimeRDG/Private/NNERuntimeRDGBase.h
--- a/Engine/Plugins/Experimental/NNERuntimeRDG/Source/NNERuntimeRDG/Private/NNERuntimeRDGBase.h
+++ b/Engine/Plugins/Experimental/NNERuntimeRDG/Source/NNERuntimeRDG/Private/NNERuntimeRDGBase.h
@@ -71,7 +71,7 @@ public:
 
 using FTensorRDGRef = FTensorRDG*;
 using FTensorRDGArray = TArray<FTensorRDG, TInlineAllocator<16>>;
-using FTensorRDGRefArray = TArray<FTensorRDGRef, TInlineAllocator<64>>;
+using FTensorRDGRefMap = TMap<int32, FTensorRDGRef>;
 using FIntArray = TArray<int32, TInlineAllocator<16>>;
 
 bool AlwaysValidValidationFunction(const NNE::FAttributeMap& AttributeMap, TConstArrayView<ENNETensorDataType> InputTypes, TConstArrayView<NNE::FSymbolicTensorShape> InputShapes);
diff --git a/Engine/Plugins/Experimental/NNERuntimeRDG/Source/NNERuntimeRDG/Private/NNERuntimeRDGModel.cpp b/Engine/Plugins/Experimental/NNERuntimeRDG/Source/NNERuntimeRDG/Private/NNERuntimeRDGModel.cpp
--- a/Engine/Plugins/Experimental/NNERuntimeRDG/Source/NNERuntimeRDG/Private/NNERuntimeRDGModel.cpp
+++ b/Engine/Plugins/Experimental/NNERuntimeRDG/Source/NNERuntimeRDG/Private/NNERuntimeRDGModel.cpp
@@ -34,6 +34,8 @@ bool FModelInstanceRDG::LoadModel(TConstArrayView<uint8> ModelData, FNNERuntimeF
 	OperatorInputTensorIndices.Empty();
 	OperatorOutputTensorIndices.Empty();
 
+	TensorIdxSpan = Format.Tensors.Num();
+
 	// Add tensors
 	for (int32 Idx = 0; Idx < Format.Tensors.Num(); ++Idx)
 	{
@@ -42,9 +44,12 @@ bool FModelInstanceRDG::LoadModel(TConstArrayView<uint8> ModelData, FNNERuntimeF
 		const NNE::FSymbolicTensorShape SymbolicShape = NNE::FSymbolicTensorShape::Make(FormatTensorDesc.Shape);
 		const NNE::FTensorDesc SymbolicTensor = NNE::FTensorDesc::Make(FormatTensorDesc.Name, SymbolicShape, FormatTensorDesc.DataType);
 
-		AllSymbolicTensorDescs.Emplace(SymbolicTensor);
+		if (Format.Tensors[Idx].Type != ENNEFormatTensorType::Empty)
+		{
+			AllSymbolicTensorDescs.Emplace(Idx, SymbolicTensor);
+		}
 
-		if (FormatTensorDesc.Type == ENNEFormatTensorType::Input || FormatTensorDesc.Type == ENNEFormatTensorType::Empty)
+		if (FormatTensorDesc.Type == ENNEFormatTensorType::Input)
 		{
 			InputTensorIndices.Emplace(Idx);
 			InputSymbolicTensors.Emplace(SymbolicTensor);
@@ -107,7 +112,7 @@ int32 FModelInstanceRDG::SetInputTensorShapes(TConstArrayView<NNE::FTensorShape>
 	}
 
 	//Allocate and prime all AllTensorRDGRefs with concrete shapes defaulting variables dimension to 1 if needed
-	AllTensorRDGRefs.Init(nullptr, AllSymbolicTensorDescs.Num());
+	AllTensorRDGRefs.Reserve(AllSymbolicTensorDescs.Num());
 
 	InputTensorRDGs.Reset(InputTensorIndices.Num());
 	for (int32 i = 0; i < InputTensorIndices.Num(); ++i)
@@ -117,14 +122,14 @@ int32 FModelInstanceRDG::SetInputTensorShapes(TConstArrayView<NNE::FTensorShape>
 		const NNE::FTensorShape& TensorShape = InputTensorShapes[i];
 
 		InputTensorRDGs.Emplace(FTensorRDG::Make(TensorDesc, TensorShape, nullptr));
-		AllTensorRDGRefs[Idx] = &InputTensorRDGs[i];
+		AllTensorRDGRefs.Emplace(Idx, &InputTensorRDGs[i]);
 	}
 
 	for (int32 i = 0; i < WeightTensorIndices.Num(); ++i)
 	{
 		const int32 Idx = WeightTensorIndices[i];
 
-		AllTensorRDGRefs[Idx] = &WeightTensorRDGs[i];
+		AllTensorRDGRefs.Emplace(Idx, &WeightTensorRDGs[i]);
 	}
 
 	IntermediateTensorRDGs.Reset(IntermediateTensorIndices.Num());
@@ -135,7 +140,7 @@ int32 FModelInstanceRDG::SetInputTensorShapes(TConstArrayView<NNE::FTensorShape>
 		const NNE::FTensorShape TensorShape = NNE::FTensorShape::MakeFromSymbolic(TensorDesc.GetShape());
 
 		IntermediateTensorRDGs.Emplace(FTensorRDG::Make(TensorDesc, TensorShape, nullptr));
-		AllTensorRDGRefs[Idx] = &IntermediateTensorRDGs[i];
+		AllTensorRDGRefs.Emplace(Idx, &IntermediateTensorRDGs[i]);
 	}
 
 	OutputTensorRDGs.Reset(OutputTensorIndices.Num());
@@ -146,14 +151,11 @@ int32 FModelInstanceRDG::SetInputTensorShapes(TConstArrayView<NNE::FTensorShape>
 		const NNE::FTensorShape TensorShape = NNE::FTensorShape::MakeFromSymbolic(TensorDesc.GetShape());
 
 		OutputTensorRDGs.Emplace(FTensorRDG::Make(TensorDesc, TensorShape, nullptr));
-		AllTensorRDGRefs[Idx] = &OutputTensorRDGs[i];
+		AllTensorRDGRefs.Emplace(Idx, &OutputTensorRDGs[i]);
 	}
 
 	checkCode(
-		for (int32 i = 0; i < AllTensorRDGRefs.Num(); ++i)
-		{
-			checkf(AllTensorRDGRefs[i] != nullptr, TEXT("Tensor at index %d, was not allocated for model preparation."), i);
-		};
+		checkf(AllTensorRDGRefs.Num() == AllSymbolicTensorDescs.Num(), TEXT("Some tensor was not allocated for model preparation."));
 	);
 
 	//Allow the specific runtime to run shape inference if supported
@@ -163,10 +165,9 @@ int32 FModelInstanceRDG::SetInputTensorShapes(TConstArrayView<NNE::FTensorShape>
 	}
 
 	checkCode(
-		for (int32 i = 0; i < AllTensorRDGRefs.Num(); ++i)
+		for (const TPair<int32, FTensorRDGRef>& Elem : AllTensorRDGRefs)
 		{
-			checkf(AllTensorRDGRefs[i] != nullptr, TEXT("Tensor at index %d, was not allocated after model preparation."), i);
-			checkf(AllTensorRDGRefs[i]->GetShape().IsCompatibleWith(AllSymbolicTensorDescs[i].GetShape()), TEXT("Tensor at index %d have a shape incompatible with model definition."), i);
+			checkf(Elem.Value->GetShape().IsCompatibleWith(AllSymbolicTensorDescs[Elem.Key].GetShape()), TEXT("Tensor at index %d have a shape incompatible with model definition."), Elem.Key);
 		};
 	);
 
@@ -242,7 +243,12 @@ int32 FModelInstanceRDG::EnqueueRDG(FRDGBuilder& RDGBuilder, TConstArrayView<NNE
 	//For now weights tensors are not uploaded to GPU thus GetBuffer will return nullptr for them.
 	if (bBuffersUploadedAndRegisteredToRDGGraph)
 	{
-		checkCode(for (const FTensorRDG* TensorRDG : AllTensorRDGRefs) { if (TensorRDG != nullptr) { check(TensorRDG->GetBuffer() != nullptr); } });
+		checkCode(
+			for (const TPair<int32, FTensorRDGRef>& TensorRDG : AllTensorRDGRefs) 
+			{ 
+				check(TensorRDG.Value->GetBuffer() != nullptr); 
+			}
+		);
 	}
 
 	// We can now dispatch operators
diff --git a/Engine/Plugins/Experimental/NNERuntimeRDG/Source/NNERuntimeRDG/Private/NNERuntimeRDGModel.h b/Engine/Plugins/Experimental/NNERuntimeRDG/Source/NNERuntimeRDG/Private/NNERuntimeRDGModel.h
--- a/Engine/Plugins/Experimental/NNERuntimeRDG/Source/NNERuntimeRDG/Private/NNERuntimeRDGModel.h
+++ b/Engine/Plugins/Experimental/NNERuntimeRDG/Source/NNERuntimeRDG/Private/NNERuntimeRDGModel.h
@@ -28,8 +28,10 @@ namespace UE::NNERuntimeRDG::Private
 		virtual bool PrepareModelRDG(FRDGBuilder& RDGBuilder) { return false; }
 		virtual void AddDispatchOps_RenderThread(FRDGBuilder& GraphBuilder) = 0;
 
+		int32 TensorIdxSpan;
+
 		//Tensor descriptor
-		TArray<NNE::FTensorDesc>	AllSymbolicTensorDescs;
+		TMap<int32, NNE::FTensorDesc>	AllSymbolicTensorDescs;
 
 		//Tensor indices for models
 		TArray<int32>				IntermediateTensorIndices;
@@ -42,7 +44,7 @@ namespace UE::NNERuntimeRDG::Private
 		TArray<TArray<uint32>>		OperatorOutputTensorIndices;
 
 		//RDG Tensors
-		FTensorRDGRefArray			AllTensorRDGRefs;
+		FTensorRDGRefMap			AllTensorRDGRefs;
 		FTensorRDGArray				InputTensorRDGs;
 		FTensorRDGArray				OutputTensorRDGs;
 		FTensorRDGArray				IntermediateTensorRDGs;
diff --git a/Engine/Plugins/Experimental/NNERuntimeRDG/Source/NNERuntimeRDG/Private/NNERuntimeRDGModelHlsl.cpp b/Engine/Plugins/Experimental/NNERuntimeRDG/Source/NNERuntimeRDG/Private/NNERuntimeRDGModelHlsl.cpp
--- a/Engine/Plugins/Experimental/NNERuntimeRDG/Source/NNERuntimeRDG/Private/NNERuntimeRDGModelHlsl.cpp
+++ b/Engine/Plugins/Experimental/NNERuntimeRDG/Source/NNERuntimeRDG/Private/NNERuntimeRDGModelHlsl.cpp
@@ -268,7 +268,7 @@ int FModelInstance::PrepareTensorShapesAndData()
 			//Operator could not prepare the output tensors, meaning we can't allocate
 			//output buffer before running the model. This runtime does not support this.
 			UE_LOG(LogNNE, Warning, TEXT("Could not deduce tensor shapes for this model during shape inference, HLSL runtime wont support the model as it need to precompute all shapes for performance reasons."));
-			AllTensorRDGRefs.Reset(AllSymbolicTensorDescs.Num());
+			AllTensorRDGRefs.Reset();
 			return -1;
 		}
 	}
