NNE Finish implementation of RDG Dml Resize
Problems:
* The 'sizes' tensor of ONNX Resize couldn't be implemented in RDG Dml because it required an infrastructural change (empty tensors) which is now available
Solution:
* Adapt the implementation and add a test
Test:
* NNE editor win64 : Run all tests again
* NNE standalone win64
* Non-unity build (p4v)

#rnx

diff --git a/Engine/Plugins/Experimental/NNERuntimeRDG/Source/NNERuntimeRDG/Private/Dml/NNEDmlOperatorResample.cpp b/Engine/Plugins/Experimental/NNERuntimeRDG/Source/NNERuntimeRDG/Private/Dml/NNEDmlOperatorResample.cpp
--- a/Engine/Plugins/Experimental/NNERuntimeRDG/Source/NNERuntimeRDG/Private/Dml/NNEDmlOperatorResample.cpp
+++ b/Engine/Plugins/Experimental/NNERuntimeRDG/Source/NNERuntimeRDG/Private/Dml/NNEDmlOperatorResample.cpp
@@ -44,6 +44,7 @@ class FOperatorDmlResample : public FOperatorDml
 
 	DML_INTERPOLATION_MODE	Mode;
 	ECoordTransformMode		CoordTransformMode;
+	bool 					bUseSizesTensor;
 
 	static DML_INTERPOLATION_MODE ModeFromString(FStringView StringVal)
 	{
@@ -76,12 +77,26 @@ public:
 
 	virtual bool Initialize(TConstArrayView<NNE::FTensorDesc> Inputs, TConstArrayView<NNE::FTensorDesc> Outputs, const NNE::FAttributeMap& Attributes) override
 	{
-		check(Inputs.Num() >= 1 && Inputs.Num() < 4);
+		check(Inputs.Num() >= 1 && Inputs.Num() <= 4);
 		check(Outputs.Num() == 1);
 
-		if (Inputs.Num() == 3 && Inputs[2].GetDataType() == ENNETensorDataType::Int64)
+		bUseSizesTensor = Inputs.Num() == 4;
+
+		if (Inputs.Num() == 3 && Inputs[2].GetDataType() != ENNETensorDataType::Float)
+		{
+			UE_LOG(LogNNE, Warning, TEXT("Unsupported input type for 'scales' of name %s, only 'scales' of type float is supported."), *Inputs[2].GetName());
+			return false;
+		}
+		
+		if (bUseSizesTensor && Inputs[3].GetDataType() != ENNETensorDataType::Int64)
+		{
+			UE_LOG(LogNNE, Warning, TEXT("Unsupported input type for 'sizes' of name %s, only 'sizes' of type int64 is supported."), *Inputs[3].GetName());
+			return false;
+		}
+
+		if(bUseSizesTensor && Inputs[2].GetName() != TEXT(""))
 		{
-			UE_LOG(LogNNE, Warning, TEXT("Unsupported input type for 'sizes' of name %s, only 'scales' of type float is supported."), *Inputs[2].GetName());
+			UE_LOG(LogNNE, Warning, TEXT("'scales' should be an empty-string tensor when 'sizes' is specified."));
 			return false;
 		}
 
@@ -92,15 +107,15 @@ public:
 		{
 			if (Mode == DML_INTERPOLATION_MODE_NEAREST_NEIGHBOR)
 			{
-				FString NeareastMode = Attributes.GetValueOrDefault<FString>(TEXT("nearest_mode"), TEXT("round_prefer_floor"));
+				// DML only supports floor. If 'nearest_mode' wasn't specified, no need to warn the user.
+				FString NeareastMode = Attributes.GetValueOrDefault<FString>(TEXT("nearest_mode"), TEXT("floor"));
 				if (FCString::Stricmp(*NeareastMode, TEXT("floor")) != 0)
 				{
-					UE_LOG(LogNNE, Warning, TEXT("Unsupported neareast mode:%s, using floor instead"), *NeareastMode);
+					UE_LOG(LogNNE, Display, TEXT("Unsupported neareast mode:%s, using floor instead"), *NeareastMode);
 				}
 			}
 
 			FString CoordinateTransformationMode = Attributes.GetValueOrDefault<FString>(TEXT("coordinate_transformation_mode"), TEXT("half_pixel"));
-			
 			if (CoordinateTransformationMode == TEXT("align_corners"))
 			{
 				CoordTransformMode = ECoordTransformMode::AlignCorners;
@@ -135,20 +150,20 @@ public:
 
 		return true;
 	}
-
+	
 	virtual int PrepareOutputs(TConstArrayView<NNE::Internal::FTensorRef> InputTensors, TArrayView<NNE::Internal::FTensorRef> OutputTensors) const override
 	{
 		const NNE::Internal::FTensor& InputTensor = *InputTensors[0];
 		const NNE::Internal::FTensor& ScaleTensor = (InputTensors.Num() == 2) ? *InputTensors[1] : *InputTensors[2]; // Upsample has scale at position 1, while Resize at 2
 		const NNE::Internal::FTensor& OutputTensor = *OutputTensors[0];
 
-		if (!ScaleTensor.HasPreparedData())
+		if (!bUseSizesTensor && !ScaleTensor.HasPreparedData())
 		{
 			UE_LOG(LogNNE, Error, TEXT("scales should be a constant tensor, it is here a variable tensor of name %s."), *ScaleTensor.GetName());
 			return -1;
 		}
 
-		if (ScaleTensor.GetShape().Volume() != InputTensor.GetShape().Rank())
+		if (!bUseSizesTensor && ScaleTensor.GetShape().Volume() != InputTensor.GetShape().Rank())
 		{
 			UE_LOG(LogNNE, Error, TEXT("scales tensor should contain N entries, where N is rank of X."));
 			return -1;
@@ -159,8 +174,8 @@ public:
 			if (CoordTransformMode == ECoordTransformMode::TfCropAndResize)
 			{
 				const NNE::Internal::FTensor& RoiTensor = *InputTensors[1];
-
-				if (!RoiTensor.HasPreparedData())
+				
+				if (!RoiTensor.IsConstant())
 				{
 					UE_LOG(LogNNE, Warning, TEXT("roi should be a constant tensor, it is here a variable tensor of name %s."), *RoiTensor.GetName());
 					return -1;
@@ -189,18 +204,41 @@ public:
 					}
 				}
 			}
-		}
 
-		TConstArrayView<float>	ScalesData = ScaleTensor.GetPreparedData<float>();
-		TConstArrayView<uint32>	InputShape = InputTensor.GetShape().GetData();
-		TArray<uint32>			OutputShape;
+			if(bUseSizesTensor)
+			{
+				const NNE::Internal::FTensor& SizesTensor = *InputTensors[2];
+				Util::FSmallArray<int64> Sizes(SizesTensor.GetPreparedData<int64>());
+				Util::FSmallArray<uint32> SizesUint32;
+				SizesUint32.SetNumUninitialized(Sizes.Num());
+				for(int Idx = 0; Idx < Sizes.Num(); ++Idx)
+				{
+					SizesUint32[Idx] = (uint32) Sizes[Idx];
+					if(Sizes[Idx] != (int64) SizesUint32[Idx])
+					{
+						UE_LOG(LogNNE, Warning, TEXT("Overflow in conversion of 'sizes'"));
+						return -1;
+					}
+				}
+				OutputTensors[0]->SetShape(NNE::FTensorShape::Make(TConstArrayView<uint32>(SizesUint32)));
+			
+			}
+		}
 
-		for (int32 i = 0; i < InputShape.Num(); ++i)
+		if(!bUseSizesTensor)
 		{
-			OutputShape.Emplace(FMath::FloorToInt32(InputShape[i] * ScalesData[i]));
-		}
+			TConstArrayView<float>	ScalesData = ScaleTensor.GetPreparedData<float>();
+			TConstArrayView<uint32>	InputShape = InputTensor.GetShape().GetData();
+			TArray<uint32>			OutputShape;
 
-		OutputTensors[0]->SetShape(NNE::FTensorShape::Make(OutputShape));
+			for (int32 i = 0; i < InputShape.Num(); ++i)
+			{
+				OutputShape.Emplace(FMath::FloorToInt32(InputShape[i] * ScalesData[i]));
+			}
+
+			OutputTensors[0]->SetShape(NNE::FTensorShape::Make(OutputShape));
+		}
+		
 
 		return 0;
 	}
@@ -216,10 +254,25 @@ public:
 		InputPixelOffsets.Init(0.5f, InputTensor.GetShape().Rank());
 		OutputPixelOffsets.Init(-0.5f, InputTensor.GetShape().Rank());
 
-		Util::FSmallArray<float> Scales(ScaleTensor.GetPreparedData<float>());
+		Util::FSmallArray<float> Scales;
+		if(!bUseSizesTensor)
+		{
+			Scales = ScaleTensor.GetPreparedData<float>();
+		}
 
 		if constexpr (IsResize)
 		{
+			if(bUseSizesTensor)
+			{
+				const NNE::Internal::FTensor& SizesTensor = *InputTensors[2];
+				Util::FSmallArray<int64> Sizes(SizesTensor.GetPreparedData<int64>());
+
+				Scales.SetNum(Sizes.Num());
+				for(int Idx = 0; Idx < Sizes.Num(); ++Idx)
+				{
+					Scales[Idx] = (float) Sizes[Idx] / (float) InputTensor.GetShape().GetData()[Idx];
+				}
+			}
 			for (int Idx = 0; Idx < InputTensor.GetShape().Rank(); ++Idx)
 			{
 				float LengthResized = (float) OutputTensor.GetShape().GetData()[Idx];
diff --git a/Engine/Source/Runtime/NNE/Internal/NNETensor.h b/Engine/Source/Runtime/NNE/Internal/NNETensor.h
--- a/Engine/Source/Runtime/NNE/Internal/NNETensor.h
+++ b/Engine/Source/Runtime/NNE/Internal/NNETensor.h
@@ -74,6 +74,11 @@ namespace UE::NNE::Internal
 			return !PreparedData.IsEmpty();
 		}
 
+		bool IsConstant() const
+		{
+			return Volume == 0 || HasPreparedData();
+		}
+
 		uint32 GetVolume() const
 		{
 			return Volume;
