[NNE] RuntimeRDGDml model instance improve CPU/GPU synchronization, GPU resource usage and D3D debug layer reports

Problems:
GPU crash can happen when running DML FModelInstance.
The problem is that RHI buffers allocated for DML inference can be deleted before GPU command list has finished using them.
There are messages from D3D debug layer:
- Resource incorrect state
- Metacommand which we don't want to display

Solution:
- Use dispatch fence to wait for GPU
- Added command lists flushing to process GPU fences and GPU resource deletion
- Add correct state for resource transition
- Disable metacommand debug logging
- Added GPU inference to NNEIncubator to NNECharacter
- Enable metacommands by default to have best performance with DirectML
- Added RDG utilities, for now FReadbackManager that handles GPU->CPU readbacks

Tests:
- Win64 Editor NNE all tests + Style transfer demo + Empty level
- Win64 standalone Style transfer demo + Empty level
- Compile non-unity

#rnx

diff --git a/Engine/Plugins/Experimental/NNERuntimeRDG/Source/NNERuntimeRDG/Private/Dml/NNEDmlModel.cpp b/Engine/Plugins/Experimental/NNERuntimeRDG/Source/NNERuntimeRDG/Private/Dml/NNEDmlModel.cpp
--- a/Engine/Plugins/Experimental/NNERuntimeRDG/Source/NNERuntimeRDG/Private/Dml/NNEDmlModel.cpp
+++ b/Engine/Plugins/Experimental/NNERuntimeRDG/Source/NNERuntimeRDG/Private/Dml/NNEDmlModel.cpp
@@ -12,11 +12,11 @@
 #endif
 
 // Register CVar for enabling / disabling DirectML meta commands
-static int32 GNNEDmlMetaCommands = 0;
+static int32 GNNEDmlMetaCommands = 1;
 static FAutoConsoleVariableRef CVarNNEDmlMetaCommands(
 	TEXT("nne.dml.MetaCommands"),
 	GNNEDmlMetaCommands,
-	TEXT("Use DirectML meta commands in the NNERuntimeRDGDml (default: 0)"),
+	TEXT("Use DirectML meta commands in the NNERuntimeRDGDml (default: 1)"),
 	ECVF_Scalability);
 
 namespace UE::NNERuntimeRDG::Private::Dml
@@ -487,6 +487,37 @@ public:
 			DmlExecFlags = DML_EXECUTION_FLAG_DISABLE_META_COMMANDS;
 		}
 
+		TComPtr<ID3D12InfoQueue>	InfoQueue;
+		bool						bIsDebugFilterApplied = false;
+
+		if (GRHIGlobals.IsDebugLayerEnabled)
+		{
+			Res = InModel->DevCtx->D3D12Device->QueryInterface(DML_PPV_ARGS(&InfoQueue));
+			if (SUCCEEDED(Res))
+			{
+				Res = InfoQueue->PushCopyOfStorageFilter();
+
+				if (SUCCEEDED(Res))
+				{
+					// Ignore the specified message ID. This may fail if the ID was introduced later and
+					// isn't recognized by the D3D runtime (e.g. running on an older build of Windows).
+					// Failure here is OK because it just means there's nothing to suppress.
+					D3D12_MESSAGE_ID IgnoredIds[] = { D3D12_MESSAGE_ID_META_COMMAND_UNSUPPORTED_PARAMS };
+
+					D3D12_INFO_QUEUE_FILTER Filter = {};
+
+					Filter.DenyList.NumIDs = UE_ARRAY_COUNT(IgnoredIds);
+					Filter.DenyList.pIDList = IgnoredIds;
+
+					Res = InfoQueue->AddStorageFilterEntries(&Filter);
+					if (SUCCEEDED(Res))
+					{
+						bIsDebugFilterApplied = true;
+					}
+				}
+			}
+		}
+
 		Res = Device1->CompileGraph(&Graph, DmlExecFlags, DML_PPV_ARGS(&Op));
 		if (FAILED(Res))
 		{
@@ -494,6 +525,12 @@ public:
 			Op = nullptr;
 		};
 
+		if (bIsDebugFilterApplied)
+		{
+			check(InfoQueue.IsValid());
+			InfoQueue->PopStorageFilter();
+		}
+
 		return Op;
 	}
 
@@ -829,6 +866,48 @@ FModelInstance::~FModelInstance()
 	DEC_MEMORY_STAT_BY(STAT_MemSizeTemp, MemSizeTemp);
 	DEC_MEMORY_STAT_BY(STAT_MemSizePersist, MemSizePersist);
 #endif
+	
+	FEvent* Signal = FGenericPlatformProcess::GetSynchEventFromPool(false);
+
+	ENQUEUE_RENDER_COMMAND(FDmlModelInstance_Release)
+	(
+		[this, &Signal](FRHICommandListImmediate& RHICmdList)
+		{
+			if (DispatchFence.IsValid())
+			{
+				while (!DispatchFence->Poll())
+				{
+					FPlatformProcess::Sleep(0.0f);
+				}
+
+				DispatchFence->DisableLifetimeExtension();
+				DispatchFence.SafeRelease();
+			}
+
+			DescHeap.Reset();
+			BindingTable.Reset();
+
+			if (PersistBuff.IsValid())
+			{
+				PersistBuff->DisableLifetimeExtension();
+				PersistBuff.SafeRelease();
+			}
+
+			if (TempBuff.IsValid())
+			{
+				TempBuff->DisableLifetimeExtension();
+				TempBuff.SafeRelease();
+			}
+
+			RHICmdList.ImmediateFlush(EImmediateFlushType::FlushRHIThreadFlushResources);
+			RHICmdList.ImmediateFlush(EImmediateFlushType::FlushRHIThread);
+
+			Signal->Trigger();
+		}
+	);
+
+	Signal->Wait();
+	FGenericPlatformProcess::ReturnSynchEventToPool(Signal);
 }
 
 bool FModelInstance::Init(TConstArrayView<uint8> ModelData, FDmlDeviceContext* InDevCtx)
@@ -953,7 +1032,7 @@ bool FModelInstance::Init(TConstArrayView<uint8> ModelData, FDmlDeviceContext* I
 bool FModelInstance::InitCompiledOp()
 {
 	static constexpr EBufferUsageFlags	WeightBuffUsage = BUF_UnorderedAccess;
-	static constexpr ERHIAccess			WeightBuffAccess = ERHIAccess::UAVMask;
+	static constexpr ERHIAccess			WeightBuffAccess = ERHIAccess::CopyDest;
 
 	static constexpr EBufferUsageFlags	PersistBuffFlags = BUF_Static | BUF_ShaderResource | BUF_UnorderedAccess;
 	static constexpr ERHIAccess			PersistBuffAccess = ERHIAccess::UAVMask;
@@ -1002,7 +1081,7 @@ bool FModelInstance::InitCompiledOp()
 	MemSizePersist = ExecBindProps.PersistentResourceSize;
 
 	FEvent* Signal = FGenericPlatformProcess::GetSynchEventFromPool(false);
-
+	
 	ENQUEUE_RENDER_COMMAND(FDmlModelInstance_SetTensorData)
 	(
 		[
@@ -1012,6 +1091,8 @@ bool FModelInstance::InitCompiledOp()
 		]
 		(FRHICommandListImmediate& RHICmdList)
 		{
+			DispatchFence = DynamicRHI->RHICreateGPUFence(TEXT("FDmlModelInstanceDispatchFence"));
+
 			FRHIBufferInputArray	Inputs;
 
 			for (int32 InputIdx : InputTensorIndices)
@@ -1024,15 +1105,19 @@ bool FModelInstance::InitCompiledOp()
 
 			TArray<CD3DX12_RESOURCE_BARRIER, TInlineAllocator<MaxNumInputs>>	Barriers;
 			FGPUFenceRHIRef	UploadFence = nullptr;
+			FBufferRHIRef	UploadBuff;
 
 			if (MemSizeWeights)
 			{
 				UploadFence = RHICreateGPUFence(TEXT("FDmlModelInstance_UploadFence"));
 
-				FBufferRHIRef	UploadBuff = CreateRHIBuffer(RHICmdList, MemSizeWeights, BUF_ShaderResource | BUF_Dynamic | BUF_FastVRAM, ERHIAccess::CopySrc, TEXT("FDmlModelInstance_UploadBuffer"));
-				uint8*			UploadBuffPtr = static_cast<uint8*>(RHICmdList.LockBuffer(UploadBuff, 0, MemSizeWeights, RLM_WriteOnly_NoOverwrite));
-				uint64			UploadOffset = 0;
+				UploadBuff = CreateRHIBuffer(RHICmdList, MemSizeWeights, BUF_ShaderResource | BUF_Dynamic | BUF_FastVRAM, ERHIAccess::CopySrc, TEXT("FDmlModelInstance_UploadBuffer"));
+				check(UploadBuff);
+				UploadBuff->DisableLifetimeExtension();
 
+				uint8*	UploadBuffPtr = static_cast<uint8*>(RHICmdList.LockBuffer(UploadBuff, 0, MemSizeWeights, RLM_WriteOnly_NoOverwrite));
+				uint64	UploadOffset = 0;
+				
 				for (int32 WeightIdx = 0; WeightIdx < WeightTensorIndices.Num(); ++WeightIdx)
 				{
 					int32 TensorIdx = WeightTensorIndices[WeightIdx];
@@ -1046,7 +1131,8 @@ bool FModelInstance::InitCompiledOp()
 
 					FBufferRHIRef WeightBuff;
 					WeightBuff = CreateRHIBuffer(RHICmdList, TensorData.Num(), WeightBuffUsage, WeightBuffAccess, *Tensor.GetName());
-				
+					check(WeightBuff);
+
 					FMemory::Memcpy(UploadBuffPtr + UploadOffset, TensorData.GetData(), TensorData.Num());
 					RHICmdList.CopyBufferRegion(WeightBuff, 0, UploadBuff, UploadOffset, TensorData.Num());
 					UploadOffset += Align(TensorData.Num(), DML_MINIMUM_BUFFER_TENSOR_ALIGNMENT);
@@ -1070,12 +1156,14 @@ bool FModelInstance::InitCompiledOp()
 			if (MemSizePersist)
 			{
 				PersistBuff = CreateRHIBuffer(RHICmdList, MemSizePersist, PersistBuffFlags, PersistBuffAccess, TEXT("FDmlModelInstance_PeristBuff"));
+				check(PersistBuff.IsValid());
 				INC_MEMORY_STAT_BY(STAT_MemSizePersist, MemSizePersist);
 			}
 
 			if (MemSizeTemp)
 			{
 				TempBuff = CreateRHIBuffer(RHICmdList, MemSizeTemp, TempBuffFlags, TempBuffAccess, TEXT("FDmlModelInstance_TempBuff"));
+				check(TempBuff.IsValid());
 				INC_MEMORY_STAT_BY(STAT_MemSizeTemp, MemSizeTemp);
 			}
 
@@ -1084,19 +1172,29 @@ bool FModelInstance::InitCompiledOp()
 			if (InitTempMemSize)
 			{
 				InitTempBuff = CreateRHIBuffer(RHICmdList, InitTempMemSize, TempBuffFlags, TempBuffAccess, TEXT("FDmlModelInstance_InitTempBuff"));
+				InitTempBuff->DisableLifetimeExtension();
+				check(InitTempBuff.IsValid());
 			}
 
-			RHICmdList.EnqueueLambda(
-				[this, Inputs, Barriers, InitTempBuff, UploadFence](FRHICommandListImmediate& RHICmdList)
+			if (UploadFence.IsValid())
+			{
+				// Flush commands
+				RHICmdList.ImmediateFlush(EImmediateFlushType::FlushRHIThread);
+
+				while (!UploadFence->Poll())
 				{
-					while (UploadFence && UploadFence->NumPendingWriteCommands.GetValue() > 0)
-					{
-						FPlatformProcess::Sleep(0.001);
-					}
+					FPlatformProcess::Sleep(0.0f);
+				}
+			}
 
+			DispatchFence->Clear();
+			RHICmdList.EnqueueLambda(
+				[this, Inputs, Barriers, InitTempBuff](FRHICommandListImmediate& RHICmdList)
+				{
 					ID3D12GraphicsCommandList* D3DCmdList = nullptr;
 
 					D3DCmdList = DynamicRHI->RHIGetGraphicsCommandList(DevCtx->DeviceIndex);
+					D3DCmdList->SetName(TEXT("FDmlModelInstance_SetTensorData_CmdList"));
 
 					BindingTable->Bind(OpInit, Inputs, PersistBuff, InitTempBuff);
 			
@@ -1111,9 +1209,22 @@ bool FModelInstance::InitCompiledOp()
 
 					DynamicRHI->RHIFinishExternalComputeWork(DevCtx->DeviceIndex, D3DCmdList);
 				}
-			);			
-			
+			);
+			RHICmdList.WriteGPUFence(DispatchFence);
+
+			// Flush commands
 			RHICmdList.ImmediateFlush(EImmediateFlushType::FlushRHIThread);
+
+			// Now we can release resources
+			if (InitTempBuff.IsValid() || UploadBuff.IsValid())
+			{
+				InitTempBuff.SafeRelease();
+				UploadBuff.SafeRelease();
+
+				RHICmdList.ImmediateFlush(EImmediateFlushType::FlushRHIThreadFlushResources);
+				RHICmdList.ImmediateFlush(EImmediateFlushType::FlushRHIThread);
+			}
+
 			Signal->Trigger();
 		}
 	);
@@ -1191,6 +1302,8 @@ void FModelInstance::AddDispatchOps_RenderThread(FRDGBuilder& GraphBuilder)
 #if HAS_GPU_STATS
 				SCOPED_GPU_STAT(RHICmdList, GPU_STAT_DispatchD3DTime)
 #endif
+
+				DispatchFence->Clear();
 				RHICmdList.EnqueueLambda(
 					[this, InputBuffers = MoveTemp(RHIInputBuffers), OutputBuffers = MoveTemp(RHIOutputBuffers)](FRHICommandListImmediate& RHICmdList)
 					{
@@ -1224,7 +1337,7 @@ void FModelInstance::AddDispatchOps_RenderThread(FRDGBuilder& GraphBuilder)
 								PreBarriers.Emplace(
 									CD3DX12_RESOURCE_BARRIER::Transition(
 										Resource,
-										D3D12_RESOURCE_STATE_COMMON,
+										D3D12_RESOURCE_STATE_COPY_SOURCE,
 										D3D12_RESOURCE_STATE_UNORDERED_ACCESS)
 								);
 
@@ -1235,23 +1348,26 @@ void FModelInstance::AddDispatchOps_RenderThread(FRDGBuilder& GraphBuilder)
 										D3D12_RESOURCE_STATE_COPY_SOURCE)
 								);
 							}
-
-							PostBarriers.Add(CD3DX12_RESOURCE_BARRIER::UAV(Resource));
 						}
 
-						BindingTable->Bind(CompiledOp, InputBuffers, OutputBuffers, PersistBuff, TempBuff);
+						PostBarriers.Add(CD3DX12_RESOURCE_BARRIER::UAV(nullptr));
 
+						BindingTable->Bind(CompiledOp, InputBuffers, OutputBuffers, PersistBuff, TempBuff);
+						
 						ID3D12GraphicsCommandList* D3DCmdList = nullptr;
 
 						D3DCmdList = DynamicRHI->RHIGetGraphicsCommandList(DevCtx->DeviceIndex);
+						D3DCmdList->SetName(TEXT("FDmlModelInstance_Dispatch_CmdList"));
 						D3DCmdList->SetDescriptorHeaps(1, &DescHeap);
 						D3DCmdList->ResourceBarrier(PreBarriers.Num(), PreBarriers.GetData());
 						DevCtx->CmdRec->RecordDispatch(D3DCmdList, CompiledOp, BindingTable->Get());
 						D3DCmdList->ResourceBarrier(PostBarriers.Num(), PostBarriers.GetData());
-
+						
 						DynamicRHI->RHIFinishExternalComputeWork(DevCtx->DeviceIndex, D3DCmdList);
 					}
 				);
+
+				RHICmdList.WriteGPUFence(DispatchFence);
 			}
 		}
 	);
diff --git a/Engine/Plugins/Experimental/NNERuntimeRDG/Source/NNERuntimeRDG/Private/Dml/NNEDmlModel.h b/Engine/Plugins/Experimental/NNERuntimeRDG/Source/NNERuntimeRDG/Private/Dml/NNEDmlModel.h
--- a/Engine/Plugins/Experimental/NNERuntimeRDG/Source/NNERuntimeRDG/Private/Dml/NNEDmlModel.h
+++ b/Engine/Plugins/Experimental/NNERuntimeRDG/Source/NNERuntimeRDG/Private/Dml/NNEDmlModel.h
@@ -118,6 +118,7 @@ private:
 	uint32								DescCount;
 	uint32								DescSize;
 
+	FGPUFenceRHIRef						DispatchFence;
 	FBufferRHIRef						PersistBuff;
 	FBufferRHIRef						TempBuff;
 	uint64								MemSizeWeights;
