[NNE] Use ID3D12InfoQueue only on Win64 platform.

Problem:
Project doesn't compile on all supported platforms

Solution:
Make sure to disable metacommand display only on Win64
Fixes compilation issue 382045

Tests:
- NNEIncubator Standalone Win64 + XSX

#rnx

diff --git a/Engine/Plugins/Experimental/NNERuntimeRDG/Source/NNERuntimeRDG/Private/Dml/NNEDmlModel.cpp b/Engine/Plugins/Experimental/NNERuntimeRDG/Source/NNERuntimeRDG/Private/Dml/NNEDmlModel.cpp
--- a/Engine/Plugins/Experimental/NNERuntimeRDG/Source/NNERuntimeRDG/Private/Dml/NNEDmlModel.cpp
+++ b/Engine/Plugins/Experimental/NNERuntimeRDG/Source/NNERuntimeRDG/Private/Dml/NNEDmlModel.cpp
@@ -487,6 +487,7 @@ public:
 			DmlExecFlags = DML_EXECUTION_FLAG_DISABLE_META_COMMANDS;
 		}
 
+#if defined(PLATFORM_WINDOWS) && !defined(PLATFORM_NNE_MICROSOFT)
 		TComPtr<ID3D12InfoQueue>	InfoQueue;
 		bool						bIsDebugFilterApplied = false;
 
@@ -517,6 +518,7 @@ public:
 				}
 			}
 		}
+#endif
 
 		Res = Device1->CompileGraph(&Graph, DmlExecFlags, DML_PPV_ARGS(&Op));
 		if (FAILED(Res))
@@ -525,11 +527,13 @@ public:
 			Op = nullptr;
 		};
 
+#if defined(PLATFORM_WINDOWS) && !defined(PLATFORM_NNE_MICROSOFT)
 		if (bIsDebugFilterApplied)
 		{
 			check(InfoQueue.IsValid());
 			InfoQueue->PopStorageFilter();
 		}
+#endif
 
 		return Op;
 	}
