[NNE] Suppress DirectML meta command unsupported params debug messages

Problem:
Currently only RDGDml runtime is disabling metacommand creation messages when -d3ddebug is enabled, but other ORT runtimes that use DML are emitting them.

Solution:
Suppress those messages globally in the D3D12Adapter.cpp.

Tests:
- NNEIncubator Editor Win64
- Compile Non-Unity

#rnx

diff --git a/Engine/Plugins/Experimental/NNERuntimeRDG/Source/NNERuntimeRDG/Private/Dml/NNEDmlModel.cpp b/Engine/Plugins/Experimental/NNERuntimeRDG/Source/NNERuntimeRDG/Private/Dml/NNEDmlModel.cpp
--- a/Engine/Plugins/Experimental/NNERuntimeRDG/Source/NNERuntimeRDG/Private/Dml/NNEDmlModel.cpp
+++ b/Engine/Plugins/Experimental/NNERuntimeRDG/Source/NNERuntimeRDG/Private/Dml/NNEDmlModel.cpp
@@ -487,39 +487,6 @@ public:
 			DmlExecFlags = DML_EXECUTION_FLAG_DISABLE_META_COMMANDS;
 		}
 
-#if defined(PLATFORM_WINDOWS) && !defined(PLATFORM_NNE_MICROSOFT)
-		TComPtr<ID3D12InfoQueue>	InfoQueue;
-		bool						bIsDebugFilterApplied = false;
-
-		if (GRHIGlobals.IsDebugLayerEnabled)
-		{
-			Res = InModel->DevCtx->D3D12Device->QueryInterface(DML_PPV_ARGS(&InfoQueue));
-			if (SUCCEEDED(Res))
-			{
-				Res = InfoQueue->PushCopyOfStorageFilter();
-
-				if (SUCCEEDED(Res))
-				{
-					// Ignore the specified message ID. This may fail if the ID was introduced later and
-					// isn't recognized by the D3D runtime (e.g. running on an older build of Windows).
-					// Failure here is OK because it just means there's nothing to suppress.
-					D3D12_MESSAGE_ID IgnoredIds[] = { D3D12_MESSAGE_ID_META_COMMAND_UNSUPPORTED_PARAMS };
-
-					D3D12_INFO_QUEUE_FILTER Filter = {};
-
-					Filter.DenyList.NumIDs = UE_ARRAY_COUNT(IgnoredIds);
-					Filter.DenyList.pIDList = IgnoredIds;
-
-					Res = InfoQueue->AddStorageFilterEntries(&Filter);
-					if (SUCCEEDED(Res))
-					{
-						bIsDebugFilterApplied = true;
-					}
-				}
-			}
-		}
-#endif
-
 		Res = Device1->CompileGraph(&Graph, DmlExecFlags, DML_PPV_ARGS(&Op));
 		if (FAILED(Res))
 		{
@@ -527,14 +494,6 @@ public:
 			Op = nullptr;
 		};
 
-#if defined(PLATFORM_WINDOWS) && !defined(PLATFORM_NNE_MICROSOFT)
-		if (bIsDebugFilterApplied)
-		{
-			check(InfoQueue.IsValid());
-			InfoQueue->PopStorageFilter();
-		}
-#endif
-
 		return Op;
 	}
 
diff --git a/Engine/Source/Runtime/D3D12RHI/Private/D3D12Adapter.cpp b/Engine/Source/Runtime/D3D12RHI/Private/D3D12Adapter.cpp
--- a/Engine/Source/Runtime/D3D12RHI/Private/D3D12Adapter.cpp
+++ b/Engine/Source/Runtime/D3D12RHI/Private/D3D12Adapter.cpp
@@ -827,6 +827,13 @@ void FD3D12Adapter::CreateRootDevice(bool bWithDebug)
 				// TODO: Remove this when the debug layers work for executions which are guarded by a fence
 				D3D12_MESSAGE_ID_INVALID_USE_OF_NON_RESIDENT_RESOURCE,
 #endif
+
+				// When optimizing the graph DirectML tries various configurations to check if meta command can be used.
+				// If a particular configuration is not supported the debug log message will be displayed.
+				// For large DirectML graphs there can be dozens of messages like this.
+				// This is safe to ignore
+				D3D12_MESSAGE_ID_META_COMMAND_UNSUPPORTED_PARAMS,
+
 			};
 
 #if PLATFORM_DESKTOP
