[NNE] RDGDml validate DML tensor data types, min/max rank and number of inputs

Problems:
Currently we only check if data type is FP32 (in CheckGenericTensor()).
We don't validate if tensor ranks are supported by DML.
We don't validate if number of inputs is supported by DML.

Solution:
Add per-op validation of:
* tensors' data types
* min/max tensor rank
* number of inputs

Note: This is not full validation, but only the first step towards it.
Note: InstanceNormalization required a fix: leading 1s were added to lower-than-4D input shapes, but they need to be put after the batch and channel dims.

Test:
* NNE editor win64 : Run all tests again
* NNE standalone win64
* Non-unity build (p4v)

#rnx

diff --git a/Engine/Plugins/Experimental/NNERuntimeRDG/Source/NNERuntimeRDG/Private/Dml/NNEDmlOperator.h b/Engine/Plugins/Experimental/NNERuntimeRDG/Source/NNERuntimeRDG/Private/Dml/NNEDmlOperator.h
--- a/Engine/Plugins/Experimental/NNERuntimeRDG/Source/NNERuntimeRDG/Private/Dml/NNEDmlOperator.h
+++ b/Engine/Plugins/Experimental/NNERuntimeRDG/Source/NNERuntimeRDG/Private/Dml/NNEDmlOperator.h
@@ -25,6 +25,7 @@ static FDmlOperator##OpName##Registrator RegisterDmlOperator##OpName;
 namespace UE::NNERuntimeRDG::Private::Dml
 {
 
+static constexpr int32  GMaxTensorRank = DML_TENSOR_DIMENSION_COUNT_MAX1;
 static constexpr uint32 NcdhwDimensionCount = 5;
 static constexpr uint32 NcdhwSpatialDimensionCount = 3;
 static constexpr uint32 NonspatialDimensionCount = 2;		//!< The batch and channel dimensions of NCW, NCHW, NCDHW....
diff --git a/Engine/Plugins/Experimental/NNERuntimeRDG/Source/NNERuntimeRDG/Private/Dml/NNEDmlOperatorActivations.cpp b/Engine/Plugins/Experimental/NNERuntimeRDG/Source/NNERuntimeRDG/Private/Dml/NNEDmlOperatorActivations.cpp
--- a/Engine/Plugins/Experimental/NNERuntimeRDG/Source/NNERuntimeRDG/Private/Dml/NNEDmlOperatorActivations.cpp
+++ b/Engine/Plugins/Experimental/NNERuntimeRDG/Source/NNERuntimeRDG/Private/Dml/NNEDmlOperatorActivations.cpp
@@ -13,10 +13,12 @@ namespace UE::NNERuntimeRDG::Private::Dml
 template
 <
 	typename DmlActivationOpDescType,
-	DML_OPERATOR_TYPE DmlActivationOpType
+	DML_OPERATOR_TYPE DmlActivationOpType,
+	TCHAR const *OpName
 >
 class FOperatorDmlActivationUnary : public FOperatorDml
 {
+	static constexpr uint32 NumAllowedInputTensors = 1, NumAllowedOutputTensors = 1;
 public:
 
 	static FOperatorDml* Create()
@@ -26,12 +28,12 @@ public:
 
 	static bool Validate(const NNE::FAttributeMap& AttributeMap, TConstArrayView<ENNETensorDataType> InputTypes, TConstArrayView<NNE::FSymbolicTensorShape> InputShapes)
 	{
-		if(InputShapes.Num() != 1)
+		if(InputShapes.Num() != NumAllowedInputTensors)
 		{
-			UE_LOG(LogNNE, Warning, TEXT("Invalid number of input tensors"));
+			UE_LOG(LogNNE, Warning, TEXT("DML %s: Invalid number of input tensors. %d provided, it should be %d."), OpName, InputShapes.Num(), NumAllowedInputTensors);
 			return false;
 		}
-		return CheckElementwiseTensor(InputTypes[0], InputShapes[0]);
+		return CheckElementwiseTensor(OpName, InputTypes[0], InputShapes[0]);
 	}
 
 	virtual ~FOperatorDmlActivationUnary() = default;
@@ -81,8 +83,8 @@ public:
 
 	virtual int PrepareOutputs(TConstArrayView<NNE::Internal::FTensorRef> InputTensors, TArrayView<NNE::Internal::FTensorRef> OutputTensors) const override
 	{
-		check(InputTensors.Num() == 1);
-		check(OutputTensors.Num() == 1);
+		check(InputTensors.Num() == NumAllowedInputTensors);
+		check(OutputTensors.Num() == NumAllowedOutputTensors);
 		OutputTensors[0]->SetShape(InputTensors[0]->GetShape());
 
 		return 0;
@@ -200,11 +202,12 @@ private:
 template
 <
 	typename DmlActivationOpDescType,
-	DML_OPERATOR_TYPE DmlActivationOpType
+	DML_OPERATOR_TYPE DmlActivationOpType,
+	TCHAR const *OpName
 >
 class FOperatorDmlActivationBinary : public FOperatorDml
 {
-
+	static constexpr uint32 NumAllowedInputTensors = 2, NumAllowedOutputTensors = 1;
 public:
 
 	static FOperatorDml* Create()
@@ -214,18 +217,18 @@ public:
 
 	static bool Validate(const NNE::FAttributeMap& AttributeMap, TConstArrayView<ENNETensorDataType> InputTypes, TConstArrayView<NNE::FSymbolicTensorShape> InputShapes)
 	{
-		if(InputShapes.Num() != 2)
+		if(InputShapes.Num() != NumAllowedInputTensors)
 		{
-			UE_LOG(LogNNE, Warning, TEXT("Invalid number of input tensors"));
+			UE_LOG(LogNNE, Warning, TEXT("DML %s: Invalid number of input tensors. %d provided, it should be %d."), OpName, InputShapes.Num(), NumAllowedInputTensors);
 			return false;
 		}
 
-		if(!CheckElementwiseTensor(InputTypes[0], InputShapes[0]))
+		if(!CheckElementwiseTensor(OpName, InputTypes[0], InputShapes[0]))
 		{
 			return false;
 		}
 
-		if(!CheckElementwiseTensor(InputTypes[1], InputShapes[1]))
+		if(!CheckElementwiseTensor(OpName, InputTypes[1], InputShapes[1]))
 		{
 			return false;
 		}
@@ -246,8 +249,8 @@ public:
 
 	virtual int PrepareOutputs(TConstArrayView<NNE::Internal::FTensorRef> InputTensors, TArrayView<NNE::Internal::FTensorRef> OutputTensors) const override
 	{
-		check(InputTensors.Num() == 2);
-		check(OutputTensors.Num() == 1);
+		check(InputTensors.Num() == NumAllowedInputTensors);
+		check(OutputTensors.Num() == NumAllowedOutputTensors);
 		OutputTensors[0]->SetShape(InputTensors[0]->GetShape());
 
 		return 0;
@@ -320,11 +323,12 @@ private:
 };
 
 #define REGISTER_OP_ACTIVATION(OpName, DmlOpName, OpClass) \
+TCHAR const Op##OpName##Name[] = TEXT(#OpName); \
 struct FDmlOperator##OpName##Registrator \
 { \
 	FDmlOperator##OpName##Registrator() \
 	{ \
-		FOperatorRegistryDml::Get()->OpAdd(TEXT(#OpName), OpClass<DML_ACTIVATION_##DmlOpName##_OPERATOR_DESC, DML_OPERATOR_ACTIVATION_##DmlOpName>::Create, OpClass<DML_ACTIVATION_##DmlOpName##_OPERATOR_DESC, DML_OPERATOR_ACTIVATION_##DmlOpName>::Validate); \
+		FOperatorRegistryDml::Get()->OpAdd(TEXT(#OpName), OpClass<DML_ACTIVATION_##DmlOpName##_OPERATOR_DESC, DML_OPERATOR_ACTIVATION_##DmlOpName, Op##OpName##Name>::Create, OpClass<DML_ACTIVATION_##DmlOpName##_OPERATOR_DESC, DML_OPERATOR_ACTIVATION_##DmlOpName, Op##OpName##Name>::Validate); \
 	} \
 }; \
 \
@@ -336,11 +340,12 @@ static FDmlOperator##OpName##Registrator RegisterDmlOperator##OpName;
 
 // Register unary activation OP with additional params
 #define REGISTER_OP_ACTIVATION_UNARY_PARAMS(OpName, DmlOpName, InAlpha, InBeta, InGamma) \
-template<> FOperatorDmlActivationUnary<DML_ACTIVATION_##DmlOpName##_OPERATOR_DESC, DML_OPERATOR_ACTIVATION_##DmlOpName>::FOperatorDmlActivationUnary() \
+REGISTER_OP_ACTIVATION_UNARY(OpName, DmlOpName) \
+template<> FOperatorDmlActivationUnary<DML_ACTIVATION_##DmlOpName##_OPERATOR_DESC, DML_OPERATOR_ACTIVATION_##DmlOpName, Op##OpName##Name>::FOperatorDmlActivationUnary() \
 	: Alpha(InAlpha), Beta(InBeta), Gamma(InGamma) \
 { \
-} \
-REGISTER_OP_ACTIVATION_UNARY(OpName, DmlOpName)
+}
+
 
 
 // Register binary activation OP 
diff --git a/Engine/Plugins/Experimental/NNERuntimeRDG/Source/NNERuntimeRDG/Private/Dml/NNEDmlOperatorBatchNormalization.cpp b/Engine/Plugins/Experimental/NNERuntimeRDG/Source/NNERuntimeRDG/Private/Dml/NNEDmlOperatorBatchNormalization.cpp
--- a/Engine/Plugins/Experimental/NNERuntimeRDG/Source/NNERuntimeRDG/Private/Dml/NNEDmlOperatorBatchNormalization.cpp
+++ b/Engine/Plugins/Experimental/NNERuntimeRDG/Source/NNERuntimeRDG/Private/Dml/NNEDmlOperatorBatchNormalization.cpp
@@ -32,36 +32,36 @@ public:
 
 	static bool Validate(const NNE::FAttributeMap& AttributeMap, TConstArrayView<ENNETensorDataType> InputTypes, TConstArrayView<NNE::FSymbolicTensorShape> InputShapes)
 	{
+		const FString OpName = TEXT("BatchNormalization");
+
 		if (InputShapes.Num() != Count)
 		{
-			UE_LOG(LogNNE, Warning, TEXT("DML input tensors must be %d"), Count);
+			UE_LOG(LogNNE, Warning, TEXT("DML %s: invalid number of input tensors. %d provided, it should be %d."), *OpName, InputShapes.Num(), Count);
         	return false;
 		}
-		
-		const int32 InputRank = InputShapes[0].Rank();
-		if (InputRank > 8)
+
+		if (!CheckGenericTensor(OpName, InputTypes[0], InputShapes[0], {ENNETensorDataType::Float, ENNETensorDataType::Half}))
 		{
-			UE_LOG(LogNNE, Warning, TEXT("DML InputTensor rank should be between 1 and 8, it's %d"), InputRank);
-        	return false;
+			return false;
 		}
 		
 		const int32 bTrainingMode = AttributeMap.GetValueOrDefault<int32>(TEXT("training_mode"), 0);
 		if (bTrainingMode)
 		{
-			UE_LOG(LogNNE, Warning, TEXT("DML BatchNormalization doesn't support training mode"));
+			UE_LOG(LogNNE, Warning, TEXT("DML %s: training mode not supported"), *OpName);
 			return false;
 		}
 
 		for (int32 Idx = X; Idx < Count; ++Idx)
 		{
-			if (!CheckElementwiseTensor(InputTypes[Idx], InputShapes[Idx]))
+			if (!CheckGenericTensor(OpName, InputTypes[Idx], InputShapes[Idx], {ENNETensorDataType::Float, ENNETensorDataType::Half}))
 			{
 				return false;
 			}
 			
 			if (!IsEqualOrBroadcastable(InputShapes[0].GetData(), InputShapes[Idx].GetData()))
 			{
-				UE_LOG(LogNNE, Warning, TEXT("DML BatchNormalization other tensors' shapes must be equal or broadcastable to input tensor's"));
+				UE_LOG(LogNNE, Warning, TEXT("DML %s: tensor %d has shape that's not equal or broadcastable to input tensor's"), *OpName, Idx);
 				return false;
 			}
 		}
diff --git a/Engine/Plugins/Experimental/NNERuntimeRDG/Source/NNERuntimeRDG/Private/Dml/NNEDmlOperatorCast.cpp b/Engine/Plugins/Experimental/NNERuntimeRDG/Source/NNERuntimeRDG/Private/Dml/NNEDmlOperatorCast.cpp
--- a/Engine/Plugins/Experimental/NNERuntimeRDG/Source/NNERuntimeRDG/Private/Dml/NNEDmlOperatorCast.cpp
+++ b/Engine/Plugins/Experimental/NNERuntimeRDG/Source/NNERuntimeRDG/Private/Dml/NNEDmlOperatorCast.cpp
@@ -10,6 +10,8 @@ namespace UE::NNERuntimeRDG::Private::Dml
 
 class FOperatorDmlCast : public FOperatorDml
 {
+	static constexpr uint32 NumAllowedInputTensors = 1, NumAllowedOutputTensors = 1;
+
 public:
 
 	static FOperatorDml* Create()
@@ -19,13 +21,18 @@ public:
 
 	static bool Validate(const NNE::FAttributeMap& AttributeMap, TConstArrayView<ENNETensorDataType> InputTypes, TConstArrayView<NNE::FSymbolicTensorShape> InputShapes)
 	{
-		if (InputShapes.Num() != 1)
+		const FString OpName = TEXT("Cast");
+
+		if(InputShapes.Num() != NumAllowedInputTensors)
 		{
-			UE_LOG(LogNNE, Warning, TEXT("There must be only 1 DML input tensor"));
-        	return false;
+			UE_LOG(LogNNE, Warning, TEXT("DML %s: Invalid number of input tensors. %d provided, it should be %d."), *OpName, InputShapes.Num(), NumAllowedInputTensors);
+			return false;
 		}
-
-		if (!CheckGenericTensor(InputTypes[0], InputShapes[0]))
+		
+		if (!CheckGenericTensor(OpName, InputTypes[0], InputShapes[0], 
+			{ENNETensorDataType::Float, ENNETensorDataType::Half, ENNETensorDataType::Int64, ENNETensorDataType::Int32, ENNETensorDataType::Int16,
+			ENNETensorDataType::Int8, ENNETensorDataType::UInt64, ENNETensorDataType::UInt32, ENNETensorDataType::UInt16, ENNETensorDataType::UInt8}
+		  	))
 		{
 			return false;
 		}
@@ -35,8 +42,8 @@ public:
 
 	virtual bool Initialize(TConstArrayView<NNE::FTensorDesc> Inputs, TConstArrayView<NNE::FTensorDesc> Outputs, const NNE::FAttributeMap& Attributes) override
 	{
-		check(Inputs.Num() == 1);
-		check(Outputs.Num() == 1);
+		check(Inputs.Num() == NumAllowedInputTensors);
+		check(Outputs.Num() == NumAllowedOutputTensors);
 		
 		TConstArrayView<int32> InputShape = Inputs[0].GetShape().GetData();
 		TConstArrayView<int32> OutputShape = Outputs[0].GetShape().GetData();
@@ -69,8 +76,8 @@ public:
 
 	virtual int PrepareOutputs(TConstArrayView<NNE::Internal::FTensorRef> InputTensors, TArrayView<NNE::Internal::FTensorRef> OutputTensors) const override
 	{
-		check(InputTensors.Num() == 1);
-		check(OutputTensors.Num() == 1);
+		check(InputTensors.Num() == NumAllowedInputTensors);
+		check(OutputTensors.Num() == NumAllowedOutputTensors);
 		OutputTensors[0]->SetShape(InputTensors[0]->GetShape());
 
 		return 0;
diff --git a/Engine/Plugins/Experimental/NNERuntimeRDG/Source/NNERuntimeRDG/Private/Dml/NNEDmlOperatorConcat.cpp b/Engine/Plugins/Experimental/NNERuntimeRDG/Source/NNERuntimeRDG/Private/Dml/NNEDmlOperatorConcat.cpp
--- a/Engine/Plugins/Experimental/NNERuntimeRDG/Source/NNERuntimeRDG/Private/Dml/NNEDmlOperatorConcat.cpp
+++ b/Engine/Plugins/Experimental/NNERuntimeRDG/Source/NNERuntimeRDG/Private/Dml/NNEDmlOperatorConcat.cpp
@@ -12,6 +12,7 @@ namespace UE::NNERuntimeRDG::Private::Dml
 
 class FOperatorDmlConcat : public FOperatorDml
 {
+	static constexpr uint32 MinAllowedInputTensors = 1, AllowedOutputTensors = 1;
 	int32 Axis;
 
 public:
@@ -23,9 +24,18 @@ public:
 
 	static bool Validate(const NNE::FAttributeMap& AttributeMap, TConstArrayView<ENNETensorDataType> InputTypes, TConstArrayView<NNE::FSymbolicTensorShape> InputShapes)
 	{
-		if (InputShapes.Num() == 0)
+		const FString OpName = TEXT("Concat");
+		
+		const TSet<ENNETensorDataType> AllowedDataTypes = 
+								{ENNETensorDataType::Double, ENNETensorDataType::Float, ENNETensorDataType::Half, 
+								 ENNETensorDataType::Int64, ENNETensorDataType::Int32, ENNETensorDataType::Int16,
+								 ENNETensorDataType::Int8, ENNETensorDataType::UInt64, ENNETensorDataType::UInt32, 
+								 ENNETensorDataType::UInt16, ENNETensorDataType::UInt8};
+
+		if (InputShapes.Num() < MinAllowedInputTensors)
 		{
-			UE_LOG(LogNNE, Warning, TEXT("DML no input tensors"));
+			UE_LOG(LogNNE, Warning, TEXT("DML %s: invalid number of input tensors. %d provided, it should be in [%d, inf]."), 
+										*OpName, InputShapes.Num(), MinAllowedInputTensors);
 			return false;
 		}
 
@@ -34,25 +44,25 @@ public:
 		int32 Axis = HandleNegativeAxis(AttributeMap.GetValue<int32>(TEXT("axis")), InputRank);
 		if (Axis >= InputRank)
 		{
-			UE_LOG(LogNNE, Warning, TEXT("DML invalid axis: %d"), Axis);
+			UE_LOG(LogNNE, Warning, TEXT("DML %s: invalid axis: %d"), *OpName, Axis);
 			return false;
 		}
 
-		if (!CheckGenericTensor(InputTypes[0], InputShapes[0]))
+		if (!CheckGenericTensor(OpName, InputTypes[0], InputShapes[0], AllowedDataTypes))
 		{
 			return false;
 		}
 		
 		for (int32 Idx = 1; Idx < InputShapes.Num(); ++Idx)
 		{
-			if (!CheckGenericTensor(InputTypes[Idx], InputShapes[Idx]))
+			if (!CheckGenericTensor(OpName, InputTypes[Idx], InputShapes[Idx], AllowedDataTypes))
 			{
 				return false;
 			}
 
 			if (InputShapes[Idx].Rank() != InputRank)
 			{
-				UE_LOG(LogNNE, Warning, TEXT("DML concat rank mismatch for tensor %d"), Idx);
+				UE_LOG(LogNNE, Warning, TEXT("DML %s: rank mismatch for tensor %d"), *OpName, Idx);
 				return false;
 			}
 
@@ -62,7 +72,7 @@ public:
 				{
 					if (InputShapes[Idx].GetData()[Dim] != InputShapes[0].GetData()[Dim])
 					{
-						UE_LOG(LogNNE, Warning, TEXT("DML concat dimension mismatch for tensor %d"), Idx);
+						UE_LOG(LogNNE, Warning, TEXT("DML %s: dimension mismatch for tensor %d"), *OpName, Idx);
 						return false;
 					}
 				}
@@ -74,8 +84,8 @@ public:
 
 	virtual bool Initialize(TConstArrayView<NNE::FTensorDesc> Inputs, TConstArrayView<NNE::FTensorDesc> Outputs, const NNE::FAttributeMap& Attributes) override
 	{
-		check(Inputs.Num() >= 1);
-		check(Outputs.Num() == 1);
+		check(Inputs.Num() >= MinAllowedInputTensors);
+		check(Outputs.Num() == AllowedOutputTensors);
 
 		const int32 InputRank = Inputs[0].GetShape().Rank();
 
diff --git a/Engine/Plugins/Experimental/NNERuntimeRDG/Source/NNERuntimeRDG/Private/Dml/NNEDmlOperatorConvolution.cpp b/Engine/Plugins/Experimental/NNERuntimeRDG/Source/NNERuntimeRDG/Private/Dml/NNEDmlOperatorConvolution.cpp
--- a/Engine/Plugins/Experimental/NNERuntimeRDG/Source/NNERuntimeRDG/Private/Dml/NNEDmlOperatorConvolution.cpp
+++ b/Engine/Plugins/Experimental/NNERuntimeRDG/Source/NNERuntimeRDG/Private/Dml/NNEDmlOperatorConvolution.cpp
@@ -116,6 +116,8 @@ class FOperatorDmlConv : public FOperatorDml
 
 	mutable FConvArgs	Args;
 
+	static constexpr int32 MinTensorRank = 3, MaxTensorRank = 5;
+	static constexpr uint32 MinAllowedInputTensors = 2, MaxAllowedInputTensors = 3;
 public:
 
 	static FOperatorDml* Create()
@@ -125,11 +127,37 @@ public:
 
 	static bool Validate(const NNE::FAttributeMap& AttributeMap, TConstArrayView<ENNETensorDataType> InputTypes, TConstArrayView<NNE::FSymbolicTensorShape> InputShapes)
 	{
+		FString OpName = Direction == DML_CONVOLUTION_DIRECTION_FORWARD ? TEXT("Conv") : TEXT("ConvTranspose");
+
+		if (InputShapes.Num() < MinAllowedInputTensors || InputShapes.Num() > MaxAllowedInputTensors)
+		{
+			UE_LOG(LogNNE, Warning, TEXT("DML %s: invalid number of input tensors. %d provided, it should be in [%d, %d]."), 
+										*OpName, InputShapes.Num(), MinAllowedInputTensors, MaxAllowedInputTensors);
+			return false;
+		}
+
+		if (!CheckGenericTensor(OpName, InputTypes[0], InputShapes[0], {ENNETensorDataType::Float, ENNETensorDataType::Half}, MinTensorRank, MaxTensorRank) || 
+			!CheckGenericTensor(OpName, InputTypes[1], InputShapes[1], {ENNETensorDataType::Float, ENNETensorDataType::Half}, MinTensorRank, MaxTensorRank))
+		{
+			return false;
+		}
+
+		if(InputShapes.Num() == 3)
+		{
+			// Bias tensor must be 1D
+			if (!CheckGenericTensor1D(OpName, InputTypes[2], InputShapes[2], {ENNETensorDataType::Float, ENNETensorDataType::Half}))
+			{
+				return false;
+			}
+		}
+
 		return true;
 	}
 
 	virtual bool Initialize(TConstArrayView<NNE::FTensorDesc> Inputs, TConstArrayView<NNE::FTensorDesc> Outputs, const NNE::FAttributeMap& Attributes) override
 	{
+		check(Inputs.Num() >= MinAllowedInputTensors && Inputs.Num() <= MaxAllowedInputTensors);
+
 		const NNE::FTensorDesc& InputTensor = Inputs[0];
 		const NNE::FTensorDesc& FilterTensor = Inputs[1];
 
@@ -163,7 +191,7 @@ public:
 		FTensorDescDml	DmlOutputTensorDesc;
 
 		if (!DmlInputTensorDesc
-				.SetTensorRank(3, 5)
+				.SetTensorRank(MinTensorRank, MaxTensorRank)
 				.SetFromTensor(InputTensor)
 				.Validate())
 		{
@@ -172,7 +200,7 @@ public:
 		}
 
 		if (!DmlFilterTensorDesc
-				.SetTensorRank(3, 5)
+				.SetTensorRank(MinTensorRank, MaxTensorRank)
 				.SetFromTensor(FilterTensor)
 				.Validate())
 		{
@@ -185,7 +213,7 @@ public:
 			const NNE::Internal::FTensor& BiasTensor = *InputTensors[2];
 
 			if (!DmlBiasTensorDesc
-					.SetTensorRank(3, 5)
+					.SetTensorRank(MinTensorRank, MaxTensorRank)
 					.SetFromTensor1D(BiasTensor, InputTensor.GetShape().Rank())
 					.Validate())
 			{
@@ -195,7 +223,7 @@ public:
 		}
 
 		if (!DmlOutputTensorDesc
-				.SetTensorRank(3, 5)
+				.SetTensorRank(MinTensorRank, MaxTensorRank)
 				.SetFromTensor(OutputTensor)
 				.Validate())
 		{
diff --git a/Engine/Plugins/Experimental/NNERuntimeRDG/Source/NNERuntimeRDG/Private/Dml/NNEDmlOperatorCumSum.cpp b/Engine/Plugins/Experimental/NNERuntimeRDG/Source/NNERuntimeRDG/Private/Dml/NNEDmlOperatorCumSum.cpp
--- a/Engine/Plugins/Experimental/NNERuntimeRDG/Source/NNERuntimeRDG/Private/Dml/NNEDmlOperatorCumSum.cpp
+++ b/Engine/Plugins/Experimental/NNERuntimeRDG/Source/NNERuntimeRDG/Private/Dml/NNEDmlOperatorCumSum.cpp
@@ -13,6 +13,7 @@ class FOperatorDmlCumSum : public FOperatorDml
 	DML_AXIS_DIRECTION	AxisDirection;
 	int32				HasExclusiveSum;
 	mutable int32		Axis;
+	static constexpr uint32 NumAllowedInputTensors = 2, NumAllowedOutputTensors = 1;
 
 public:
 
@@ -23,13 +24,38 @@ public:
 
 	static bool Validate(const NNE::FAttributeMap& AttributeMap, TConstArrayView<ENNETensorDataType> InputTypes, TConstArrayView<NNE::FSymbolicTensorShape> InputShapes)
 	{
+		const FString OpName = TEXT("CumSum");
+
+		if(InputShapes.Num() != NumAllowedInputTensors)
+		{
+			UE_LOG(LogNNE, Warning, TEXT("DML %s: Invalid number of input tensors. %d provided, it should be %d."), *OpName, InputShapes.Num(), NumAllowedInputTensors);
+			return false;
+		}
+		
+		if (!CheckGenericTensor(OpName, InputTypes[0], InputShapes[0], 
+			{ ENNETensorDataType::Float, ENNETensorDataType::Half, 
+			  ENNETensorDataType::Int64, ENNETensorDataType::Int32, 
+			  ENNETensorDataType::UInt64, ENNETensorDataType::UInt32 }
+		  	))
+		{
+			return false;
+		}
+
+		//axis is scalar
+		if (!CheckGenericTensor1D(OpName, InputTypes[1], InputShapes[1], 
+			{ ENNETensorDataType::Int64, ENNETensorDataType::Int32 }
+		  	))
+		{
+			return false;
+		}
+
 		return true;
 	}
 
 	virtual bool Initialize(TConstArrayView<NNE::FTensorDesc> Inputs, TConstArrayView<NNE::FTensorDesc> Outputs, const NNE::FAttributeMap& Attributes) override
 	{
-		check(Inputs.Num() == 2);
-		check(Outputs.Num() == 1);
+		check(Inputs.Num() == NumAllowedInputTensors);
+		check(Outputs.Num() == NumAllowedOutputTensors);
 
 		ConstantCPUInputs.Add(1);
 
diff --git a/Engine/Plugins/Experimental/NNERuntimeRDG/Source/NNERuntimeRDG/Private/Dml/NNEDmlOperatorDepthToSpace.cpp b/Engine/Plugins/Experimental/NNERuntimeRDG/Source/NNERuntimeRDG/Private/Dml/NNEDmlOperatorDepthToSpace.cpp
--- a/Engine/Plugins/Experimental/NNERuntimeRDG/Source/NNERuntimeRDG/Private/Dml/NNEDmlOperatorDepthToSpace.cpp
+++ b/Engine/Plugins/Experimental/NNERuntimeRDG/Source/NNERuntimeRDG/Private/Dml/NNEDmlOperatorDepthToSpace.cpp
@@ -30,6 +30,8 @@ class FOperatorDmlDepthToSpace : public FOperatorDml
 
 	DML_DEPTH_SPACE_ORDER	Order;
 	int32					BlockSize;
+	static constexpr uint32 NumAllowedInputTensors = 1, NumAllowedOutputTensors = 1;
+	static constexpr int32 	MinTensorRank = 4, MaxTensorRank = 4;
 
 public:
 
@@ -40,13 +42,33 @@ public:
 
 	static bool Validate(const NNE::FAttributeMap& AttributeMap, TConstArrayView<ENNETensorDataType> InputTypes, TConstArrayView<NNE::FSymbolicTensorShape> InputShapes)
 	{
+		const FString OpName = TEXT("DepthToSpace");
+
+		if(InputShapes.Num() != NumAllowedInputTensors)
+		{
+			UE_LOG(LogNNE, Warning, TEXT("DML %s: Invalid number of input tensors. %d provided, it should be %d."), *OpName, InputShapes.Num(), NumAllowedInputTensors);
+			return false;
+		}
+		
+		if (!CheckGenericTensor(OpName, InputTypes[0], InputShapes[0], 
+			{ 	ENNETensorDataType::Double, ENNETensorDataType::Float, ENNETensorDataType::Half, 
+				ENNETensorDataType::Int64, ENNETensorDataType::Int32, ENNETensorDataType::Int16,
+				ENNETensorDataType::Int8, ENNETensorDataType::UInt64, ENNETensorDataType::UInt32, 
+				ENNETensorDataType::UInt16, ENNETensorDataType::UInt8
+			},
+			MinTensorRank, MaxTensorRank
+		  	))
+		{
+			return false;
+		}
+
 		return true;
 	}
 
 	virtual bool Initialize(TConstArrayView<NNE::FTensorDesc> Inputs, TConstArrayView<NNE::FTensorDesc> Outputs, const NNE::FAttributeMap& Attributes) override
 	{
-		check(Inputs.Num() == 1);
-		check(Outputs.Num() == 1);
+		check(Inputs.Num() == NumAllowedInputTensors);
+		check(Outputs.Num() == NumAllowedOutputTensors);
 
 		const FNNEAttributeValue* BlockSizeAttr = Attributes.GetAttributeValue(TEXT("blocksize"));
 		if (BlockSizeAttr)
@@ -90,7 +112,7 @@ public:
 		FTensorDescDml	DmlOutputTensorDesc;
 
 		if (!DmlInputTensorDesc
-				.SetTensorRank(4, 4)
+				.SetTensorRank(MinTensorRank, MaxTensorRank)
 				.SetFromTensor(InputTensorDesc)
 				.Validate())
 		{
@@ -99,7 +121,7 @@ public:
 		}
 
 		if (!DmlOutputTensorDesc
-				.SetTensorRank(4, 4)
+				.SetTensorRank(MinTensorRank, MaxTensorRank)
 				.SetFromTensor(OutputTensorDesc)
 				.Validate())
 		{
diff --git a/Engine/Plugins/Experimental/NNERuntimeRDG/Source/NNERuntimeRDG/Private/Dml/NNEDmlOperatorElementwise.cpp b/Engine/Plugins/Experimental/NNERuntimeRDG/Source/NNERuntimeRDG/Private/Dml/NNEDmlOperatorElementwise.cpp
--- a/Engine/Plugins/Experimental/NNERuntimeRDG/Source/NNERuntimeRDG/Private/Dml/NNEDmlOperatorElementwise.cpp
+++ b/Engine/Plugins/Experimental/NNERuntimeRDG/Source/NNERuntimeRDG/Private/Dml/NNEDmlOperatorElementwise.cpp
@@ -13,10 +13,12 @@ namespace UE::NNERuntimeRDG::Private::Dml
 template
 <
 	typename DmlElementWiseOpDescType, 
-	DML_OPERATOR_TYPE DmlElementWiseOpType
+	DML_OPERATOR_TYPE DmlElementWiseOpType,
+	TCHAR const *OpName
 >
 class FOperatorDmlElementWiseUnary : public FOperatorDml
 {
+	static constexpr uint32 NumAllowedInputTensors = 1, NumAllowedOutputTensors = 1;
 public:
 
 	static FOperatorDml* Create()
@@ -26,18 +28,12 @@ public:
 
 	static bool Validate(const NNE::FAttributeMap& AttributeMap, TConstArrayView<ENNETensorDataType> InputTypes, TConstArrayView<NNE::FSymbolicTensorShape> InputShapes)
 	{
-		if (InputShapes.Num() != 1)
+		if(InputShapes.Num() != NumAllowedInputTensors)
 		{
-			UE_LOG(LogNNE, Warning, TEXT("Invalid number of input tensors"));
+			UE_LOG(LogNNE, Warning, TEXT("DML %s: Invalid number of input tensors. %d provided, it should be %d."), OpName, InputShapes.Num(), NumAllowedInputTensors);
 			return false;
 		}
-
-		if (!CheckElementwiseTensor(InputTypes[0], InputShapes[0]))
-		{
-			return false;
-		}
-
-		return true;
+		return CheckElementwiseTensor(OpName, InputTypes[0], InputShapes[0]);
 	}
 
 	virtual ~FOperatorDmlElementWiseUnary() = default;
@@ -88,8 +84,8 @@ public:
 
 	virtual int PrepareOutputs(TConstArrayView<NNE::Internal::FTensorRef> InputTensors, TArrayView<NNE::Internal::FTensorRef> OutputTensors) const override
 	{
-		check(InputTensors.Num() == 1);
-		check(OutputTensors.Num() == 1);
+		check(InputTensors.Num() == NumAllowedInputTensors);
+		check(OutputTensors.Num() == NumAllowedOutputTensors);
 		OutputTensors[0]->SetShape(InputTensors[0]->GetShape());
 
 		return 0;
@@ -148,11 +144,12 @@ private:
 template
 <
 	typename DmlElementWiseOpDescType, 
-	DML_OPERATOR_TYPE DmlElementWiseOpType
+	DML_OPERATOR_TYPE DmlElementWiseOpType,
+	TCHAR const *OpName
 >
 class FOperatorDmlElementWiseBinary : public FOperatorDml
 {
-
+	static constexpr uint32 NumAllowedInputTensors = 2, NumAllowedOutputTensors = 1;
 public:
 
 	static FOperatorDml* Create()
@@ -162,18 +159,18 @@ public:
 
 	static bool Validate(const NNE::FAttributeMap& AttributeMap, TConstArrayView<ENNETensorDataType> InputTypes, TConstArrayView<NNE::FSymbolicTensorShape> InputShapes)
 	{
-		if (InputShapes.Num() != 2)
+		if(InputShapes.Num() != NumAllowedInputTensors)
 		{
-			UE_LOG(LogNNE, Warning, TEXT("Invalid number of input tensors"));
+			UE_LOG(LogNNE, Warning, TEXT("DML %s: Invalid number of input tensors. %d provided, it should be %d."), OpName, InputShapes.Num(), NumAllowedInputTensors);
 			return false;
 		}
 
-		if (!CheckElementwiseTensor(InputTypes[0], InputShapes[0]))
+		if(!CheckElementwiseTensor(OpName, InputTypes[0], InputShapes[0]))
 		{
 			return false;
 		}
 
-		if (!CheckElementwiseTensor(InputTypes[1], InputShapes[1]))
+		if(!CheckElementwiseTensor(OpName, InputTypes[1], InputShapes[1]))
 		{
 			return false;
 		}
@@ -194,8 +191,8 @@ public:
 
 	virtual int PrepareOutputs(TConstArrayView<NNE::Internal::FTensorRef> InputTensors, TArrayView<NNE::Internal::FTensorRef> OutputTensors) const override
 	{
-		check(InputTensors.Num() == 2);
-		check(OutputTensors.Num() == 1);
+		check(InputTensors.Num() == NumAllowedInputTensors);
+		check(OutputTensors.Num() == NumAllowedOutputTensors);
 
 		TConstArrayView<uint32> ShapeA = InputTensors[0]->GetShape().GetData();
 		TConstArrayView<uint32> ShapeB = InputTensors[1]->GetShape().GetData();
@@ -291,11 +288,12 @@ private:
 };
 
 #define REGISTER_OP_ELEMENT_WISE(OpName, DmlOpName, OpClass) \
+TCHAR const Op##OpName##Name[] = TEXT(#OpName); \
 struct FDmlOperator##OpName##Registrator \
 { \
 	FDmlOperator##OpName##Registrator() \
 	{ \
-		FOperatorRegistryDml::Get()->OpAdd(TEXT(#OpName), OpClass<DML_ELEMENT_WISE_##DmlOpName##_OPERATOR_DESC, DML_OPERATOR_ELEMENT_WISE_##DmlOpName>::Create, OpClass<DML_ELEMENT_WISE_##DmlOpName##_OPERATOR_DESC, DML_OPERATOR_ELEMENT_WISE_##DmlOpName>::Validate); \
+		FOperatorRegistryDml::Get()->OpAdd(TEXT(#OpName), OpClass<DML_ELEMENT_WISE_##DmlOpName##_OPERATOR_DESC, DML_OPERATOR_ELEMENT_WISE_##DmlOpName, Op##OpName##Name>::Create, OpClass<DML_ELEMENT_WISE_##DmlOpName##_OPERATOR_DESC, DML_OPERATOR_ELEMENT_WISE_##DmlOpName, Op##OpName##Name>::Validate); \
 	} \
 }; \
 \
diff --git a/Engine/Plugins/Experimental/NNERuntimeRDG/Source/NNERuntimeRDG/Private/Dml/NNEDmlOperatorFlatten.cpp b/Engine/Plugins/Experimental/NNERuntimeRDG/Source/NNERuntimeRDG/Private/Dml/NNEDmlOperatorFlatten.cpp
--- a/Engine/Plugins/Experimental/NNERuntimeRDG/Source/NNERuntimeRDG/Private/Dml/NNEDmlOperatorFlatten.cpp
+++ b/Engine/Plugins/Experimental/NNERuntimeRDG/Source/NNERuntimeRDG/Private/Dml/NNEDmlOperatorFlatten.cpp
@@ -11,6 +11,8 @@ namespace UE::NNERuntimeRDG::Private::Dml
 class FOperatorDmlFlatten : public FOperatorDml
 {
 	int32 Axis;
+	static constexpr uint32 NumAllowedInputTensors = 1, NumAllowedOutputTensors = 1;
+	static constexpr int32 	MinTensorRank = 0, MaxTensorRank = GMaxTensorRank;
 
 public:
 
@@ -21,13 +23,33 @@ public:
 
 	static bool Validate(const NNE::FAttributeMap& AttributeMap, TConstArrayView<ENNETensorDataType> InputTypes, TConstArrayView<NNE::FSymbolicTensorShape> InputShapes)
 	{
+		const FString OpName = TEXT("Flatten");
+		
+		if(InputShapes.Num() != NumAllowedInputTensors)
+		{
+			UE_LOG(LogNNE, Warning, TEXT("DML %s: Invalid number of input tensors. %d provided, it should be %d."), *OpName, InputShapes.Num(), NumAllowedInputTensors);
+			return false;
+		}
+		
+		if (!CheckGenericTensor(OpName, InputTypes[0], InputShapes[0], 
+			{ 	ENNETensorDataType::Double, ENNETensorDataType::Float, ENNETensorDataType::Half, 
+				ENNETensorDataType::Int64, ENNETensorDataType::Int32, ENNETensorDataType::Int16,
+				ENNETensorDataType::Int8, ENNETensorDataType::UInt64, ENNETensorDataType::UInt32, 
+				ENNETensorDataType::UInt16, ENNETensorDataType::UInt8
+			},
+			MinTensorRank, MaxTensorRank
+		  	))
+		{
+			return false;
+		}
+
 		return true;
 	}
 
 	virtual bool Initialize(TConstArrayView<NNE::FTensorDesc> Inputs, TConstArrayView<NNE::FTensorDesc> Outputs, const NNE::FAttributeMap& Attributes) override
 	{
-		check(Inputs.Num() == 1);
-		check(Outputs.Num() == 1);
+		check(Inputs.Num() == NumAllowedInputTensors);
+		check(Outputs.Num() == NumAllowedOutputTensors);
 
 		const int32 InputRank = Inputs[0].GetShape().Rank();
 
diff --git a/Engine/Plugins/Experimental/NNERuntimeRDG/Source/NNERuntimeRDG/Private/Dml/NNEDmlOperatorGather.cpp b/Engine/Plugins/Experimental/NNERuntimeRDG/Source/NNERuntimeRDG/Private/Dml/NNEDmlOperatorGather.cpp
--- a/Engine/Plugins/Experimental/NNERuntimeRDG/Source/NNERuntimeRDG/Private/Dml/NNEDmlOperatorGather.cpp
+++ b/Engine/Plugins/Experimental/NNERuntimeRDG/Source/NNERuntimeRDG/Private/Dml/NNEDmlOperatorGather.cpp
@@ -10,7 +10,8 @@ namespace UE::NNERuntimeRDG::Private::Dml
 class FOperatorDmlGather : public FOperatorDml
 {
 	int32	Axis;
-
+	static constexpr uint32 NumAllowedInputTensors = 2, NumAllowedOutputTensors = 1;
+	static constexpr int32 	MinTensorRank = 0, MaxTensorRank = GMaxTensorRank;
 public:
 
 	static FOperatorDml* Create()
@@ -20,13 +21,43 @@ public:
 
 	static bool Validate(const NNE::FAttributeMap& AttributeMap, TConstArrayView<ENNETensorDataType> InputTypes, TConstArrayView<NNE::FSymbolicTensorShape> InputShapes)
 	{
+		const FString OpName = TEXT("Gather");
+
+		if(InputShapes.Num() != NumAllowedInputTensors)
+		{
+			UE_LOG(LogNNE, Warning, TEXT("DML %s: Invalid number of input tensors. %d provided, it should be %d."), *OpName, InputShapes.Num(), NumAllowedInputTensors);
+			return false;
+		}
+		
+		if (!CheckGenericTensor(OpName, InputTypes[0], InputShapes[0], 
+			{ 	ENNETensorDataType::Double, ENNETensorDataType::Float, ENNETensorDataType::Half, 
+				ENNETensorDataType::Int64, ENNETensorDataType::Int32, ENNETensorDataType::Int16,
+				ENNETensorDataType::Int8, ENNETensorDataType::UInt64, ENNETensorDataType::UInt32, 
+				ENNETensorDataType::UInt16, ENNETensorDataType::UInt8
+			},
+			MinTensorRank, MaxTensorRank
+		  	))
+		{
+			return false;
+		}
+
+		if (!CheckGenericTensor(OpName, InputTypes[1], InputShapes[1], 
+			{ 	ENNETensorDataType::Int64, ENNETensorDataType::Int32, 
+				ENNETensorDataType::UInt64, ENNETensorDataType::UInt32
+			},
+			MinTensorRank, MaxTensorRank
+		  	))
+		{
+			return false;
+		}
+
 		return true;
 	}
 
 	virtual bool Initialize(TConstArrayView<NNE::FTensorDesc> Inputs, TConstArrayView<NNE::FTensorDesc> Outputs, const NNE::FAttributeMap& Attributes) override
 	{
-		check(Inputs.Num() == 2);
-		check(Outputs.Num() == 1);
+		check(Inputs.Num() == NumAllowedInputTensors);
+		check(Outputs.Num() == NumAllowedOutputTensors);
 
 		const NNE::FTensorDesc& InputTensor = Inputs[0];
 		const NNE::FTensorDesc& IndicesTensor = Inputs[1];
diff --git a/Engine/Plugins/Experimental/NNERuntimeRDG/Source/NNERuntimeRDG/Private/Dml/NNEDmlOperatorGemm.cpp b/Engine/Plugins/Experimental/NNERuntimeRDG/Source/NNERuntimeRDG/Private/Dml/NNEDmlOperatorGemm.cpp
--- a/Engine/Plugins/Experimental/NNERuntimeRDG/Source/NNERuntimeRDG/Private/Dml/NNEDmlOperatorGemm.cpp
+++ b/Engine/Plugins/Experimental/NNERuntimeRDG/Source/NNERuntimeRDG/Private/Dml/NNEDmlOperatorGemm.cpp
@@ -2,12 +2,15 @@
 
 #ifdef NNE_USE_DIRECTML
 #include "NNEDmlOperator.h"
+#include "NNEDmlOperatorUtils.h"
 
 namespace UE::NNERuntimeRDG::Private::Dml
 {
 
 class FOperatorDmlGemm : public FOperatorDml
 {
+	static constexpr uint32 MinAllowedInputTensors = 2, MaxAllowedInputTensors = 3, NumAllowedOutputTensors = 1;
+	static constexpr int32 MinTensorRank = 2, MaxTensorRank = 4;
 	float Alpha = 1.0f;
 	float Beta = 1.0f;
 	int32 TransA = 0;
@@ -22,6 +25,42 @@ public:
 
 	static bool Validate(const NNE::FAttributeMap& AttributeMap, TConstArrayView<ENNETensorDataType> InputTypes, TConstArrayView<NNE::FSymbolicTensorShape> InputShapes)
 	{
+		const FString OpName = TEXT("Gemm");
+
+		if (InputShapes.Num() < MinAllowedInputTensors || InputShapes.Num() > MaxAllowedInputTensors)
+		{
+			UE_LOG(LogNNE, Warning, TEXT("DML %s: invalid number of input tensors. %d provided, it should be in [%d, %d]."), 
+										*OpName, InputShapes.Num(), MinAllowedInputTensors, MaxAllowedInputTensors);
+			return false;
+		}
+
+		if (!CheckGenericTensor(OpName, InputTypes[0], InputShapes[0], 
+			{ ENNETensorDataType::Float, ENNETensorDataType::Half },
+			MinTensorRank, MaxTensorRank
+		  	))
+		{
+			return false;
+		}
+
+		if (!CheckGenericTensor(OpName, InputTypes[1], InputShapes[1], 
+			{ ENNETensorDataType::Float, ENNETensorDataType::Half },
+			MinTensorRank, MaxTensorRank
+		  	))
+		{
+			return false;
+		}
+
+		if(InputShapes.Num() == 3)
+		{
+			if (!CheckGenericTensor(OpName, InputTypes[2], InputShapes[2], 
+				{ ENNETensorDataType::Float, ENNETensorDataType::Half },
+				0, MaxTensorRank
+				))
+			{
+				return false;
+			}
+		}
+		
 		return true;
 	}
 
@@ -35,15 +74,15 @@ public:
 		const NNE::FTensorDesc& InputA = Inputs[0];
 		const NNE::FTensorDesc& InputB = Inputs[1];
 
-		if (InputA.GetShape().Rank() < 2 || InputA.GetShape().Rank() > 4)
+		if (InputA.GetShape().Rank() < MinTensorRank || InputA.GetShape().Rank() > MaxTensorRank)
 		{
-			UE_LOG(LogNNE, Error, TEXT("Gemm InputA tensor rank needs to be [2,4]"));
+			UE_LOG(LogNNE, Error, TEXT("Gemm InputA tensor rank needs to be [%d,%d]"), MinTensorRank, MaxTensorRank);
 			return false;
 		}
 
-		if (InputB.GetShape().Rank() < 2 || InputB.GetShape().Rank() > 4)
+		if (InputB.GetShape().Rank() < MinTensorRank || InputB.GetShape().Rank() > MaxTensorRank)
 		{
-			UE_LOG(LogNNE, Error, TEXT("Gemm InputB tensor rank needs to be [2,4]"));
+			UE_LOG(LogNNE, Error, TEXT("Gemm InputB tensor rank needs to be [%d,%d]"), MinTensorRank, MaxTensorRank);
 			return false;
 		}
 
@@ -51,9 +90,9 @@ public:
 		{
 			const NNE::FTensorDesc& InputC = Inputs[2];
 
-			if (InputC.GetShape().Rank() > 4)
+			if (InputC.GetShape().Rank() > MaxTensorRank)
 			{
-				UE_LOG(LogNNE, Error, TEXT("Gemm InputC tensor rank needs to be max rank 4"));
+				UE_LOG(LogNNE, Error, TEXT("Gemm InputC tensor rank needs to be max rank %d"), MaxTensorRank);
 				return false;
 			}
 		}
@@ -63,14 +102,14 @@ public:
 
 	virtual int PrepareOutputs(TConstArrayView<NNE::Internal::FTensorRef> InputTensors, TArrayView<NNE::Internal::FTensorRef> OutputTensors) const override
 	{
-		check(InputTensors.Num() >= 2 && InputTensors.Num() <= 3);
-		check(OutputTensors.Num() == 1);
+		check(InputTensors.Num() >= MinAllowedInputTensors && InputTensors.Num() <= MaxAllowedInputTensors);
+		check(OutputTensors.Num() == NumAllowedOutputTensors);
 
 		const NNE::FTensorShape& InputA = InputTensors[0]->GetShape();
 		const NNE::FTensorShape& InputB = InputTensors[1]->GetShape();
 		
-		checkf(InputA.Rank() >= 2, TEXT("Gemm InputA needs to have tensor rank at least size of 2"));
-		checkf(InputB.Rank() >= 2, TEXT("Gemm InputB needs to have tensor rank at least size of 2"));
+		checkf(InputA.Rank() >= MinTensorRank, TEXT("Gemm InputA needs to have tensor rank at least size of %d"), MinTensorRank);
+		checkf(InputB.Rank() >= MinTensorRank, TEXT("Gemm InputB needs to have tensor rank at least size of %d"), MinTensorRank);
 
 		const uint32 M = TransA != 0 ? InputA.GetData()[1] : InputA.GetData()[0];
 		const uint32 N = TransB != 0 ? InputB.GetData()[0] : InputB.GetData()[1];
@@ -95,7 +134,7 @@ public:
 		FTensorDescDml	DmlOutputTensorDesc;
 
 		if (!DmlInputATensorDesc
-				.SetTensorRank(2, 4)
+				.SetTensorRank(MinTensorRank, MaxTensorRank)
 				.SetFromTensor(InputATensor)
 				.Validate())
 		{
@@ -104,7 +143,7 @@ public:
 		}
 
 		if (!DmlInputBTensorDesc
-				.SetTensorRank(2, 4)
+				.SetTensorRank(MinTensorRank, MaxTensorRank)
 				.SetFromTensor(InputBTensor)
 				.Validate())
 		{
@@ -117,7 +156,7 @@ public:
 			const NNE::Internal::FTensor& InputCTensor = *InputTensors[2];
 
 			if (!DmlInputCTensorDesc
-					.SetTensorRank(2, 4)
+					.SetTensorRank(MinTensorRank, MaxTensorRank)
 					.SetFromTensorBroadcast(InputCTensor, OutputTensor.GetShape())
 					.Validate())
 			{
@@ -127,7 +166,7 @@ public:
 		}
 
 		if (!DmlOutputTensorDesc
-				.SetTensorRank(2, 4)
+				.SetTensorRank(MinTensorRank, MaxTensorRank)
 				.SetFromTensor(OutputTensor)
 				.Validate())
 		{
diff --git a/Engine/Plugins/Experimental/NNERuntimeRDG/Source/NNERuntimeRDG/Private/Dml/NNEDmlOperatorIdentity.cpp b/Engine/Plugins/Experimental/NNERuntimeRDG/Source/NNERuntimeRDG/Private/Dml/NNEDmlOperatorIdentity.cpp
--- a/Engine/Plugins/Experimental/NNERuntimeRDG/Source/NNERuntimeRDG/Private/Dml/NNEDmlOperatorIdentity.cpp
+++ b/Engine/Plugins/Experimental/NNERuntimeRDG/Source/NNERuntimeRDG/Private/Dml/NNEDmlOperatorIdentity.cpp
@@ -11,6 +11,9 @@ namespace UE::NNERuntimeRDG::Private::Dml
 class FOperatorDmlIdentity : public FOperatorDml
 {
 
+	static constexpr uint32 NumAllowedInputTensors = 1;
+	static constexpr int32 MinTensorRank = 0, MaxTensorRank = GMaxTensorRank;
+
 public:
 
 	static FOperatorDml* Create()
@@ -20,13 +23,33 @@ public:
 
 	static bool Validate(const NNE::FAttributeMap& AttributeMap, TConstArrayView<ENNETensorDataType> InputTypes, TConstArrayView<NNE::FSymbolicTensorShape> InputShapes)
 	{
+		const FString OpName = TEXT("Identity");
+
+		if(InputShapes.Num() != NumAllowedInputTensors)
+		{
+			UE_LOG(LogNNE, Warning, TEXT("DML %s: Invalid number of input tensors. %d provided, it should be %d."), *OpName, InputShapes.Num(), NumAllowedInputTensors);
+			return false;
+		}
+		
+		if (!CheckGenericTensor(OpName, InputTypes[0], InputShapes[0], 
+			{ 	ENNETensorDataType::Double, ENNETensorDataType::Float, ENNETensorDataType::Half, 
+				ENNETensorDataType::Int64, ENNETensorDataType::Int32, ENNETensorDataType::Int16,
+				ENNETensorDataType::Int8, ENNETensorDataType::UInt64, ENNETensorDataType::UInt32, 
+				ENNETensorDataType::UInt16, ENNETensorDataType::UInt8
+			},
+			MinTensorRank, MaxTensorRank
+		  	))
+		{
+			return false;
+		}
+
 		return true;
 	}
 
 	virtual bool Initialize(TConstArrayView<NNE::FTensorDesc> Inputs, TConstArrayView<NNE::FTensorDesc> Outputs, const NNE::FAttributeMap& Attributes) override
 	{
-		check(Inputs.Num() == 1);
-		check(Outputs.Num() == 1);
+		check(Inputs.Num() == NumAllowedInputTensors);
+		check(Outputs.Num() == NumAllowedInputTensors);
 		
 		return true;
 	}
diff --git a/Engine/Plugins/Experimental/NNERuntimeRDG/Source/NNERuntimeRDG/Private/Dml/NNEDmlOperatorInstanceNormalization.cpp b/Engine/Plugins/Experimental/NNERuntimeRDG/Source/NNERuntimeRDG/Private/Dml/NNEDmlOperatorInstanceNormalization.cpp
--- a/Engine/Plugins/Experimental/NNERuntimeRDG/Source/NNERuntimeRDG/Private/Dml/NNEDmlOperatorInstanceNormalization.cpp
+++ b/Engine/Plugins/Experimental/NNERuntimeRDG/Source/NNERuntimeRDG/Private/Dml/NNEDmlOperatorInstanceNormalization.cpp
@@ -2,6 +2,7 @@
 
 #ifdef NNE_USE_DIRECTML
 #include "NNEDmlOperator.h"
+#include "NNEDmlOperatorUtils.h"
 
 namespace UE::NNERuntimeRDG::Private::Dml
 {
@@ -9,8 +10,9 @@ namespace UE::NNERuntimeRDG::Private::Dml
 class FOperatorDmlInstanceNormalization : public FOperatorDml
 {
 	static constexpr float DefaultEpsilon = 0.00001f;
-
 	float Epsilon;
+	static constexpr uint32 NumAllowedInputTensors = 3, NumAllowedOutputTensors = 1;
+	static constexpr int32 	MinTensorRank = 2, MaxTensorRank = 4;
 
 public:
 
@@ -21,13 +23,46 @@ public:
 
 	static bool Validate(const NNE::FAttributeMap& AttributeMap, TConstArrayView<ENNETensorDataType> InputTypes, TConstArrayView<NNE::FSymbolicTensorShape> InputShapes)
 	{
+		const FString OpName = TEXT("InstanceNormalization");
+
+		if(InputShapes.Num() != NumAllowedInputTensors)
+		{
+			UE_LOG(LogNNE, Warning, TEXT("DML %s: Invalid number of input tensors. %d provided, it should be %d."), *OpName, InputShapes.Num(), NumAllowedInputTensors);
+			return false;
+		}
+		
+		if (!CheckGenericTensor(OpName, InputTypes[0], InputShapes[0], 
+			{ 	ENNETensorDataType::Float, ENNETensorDataType::Half
+			},
+			MinTensorRank, MaxTensorRank
+		  	))
+		{
+			return false;
+		}
+
+		if (!CheckGenericTensor1D(OpName, InputTypes[1], InputShapes[1], 
+			{ 	ENNETensorDataType::Float, ENNETensorDataType::Half
+			}
+		  	))
+		{
+			return false;
+		}
+
+		if (!CheckGenericTensor1D(OpName, InputTypes[2], InputShapes[2], 
+			{ 	ENNETensorDataType::Float, ENNETensorDataType::Half
+			}
+		  	))
+		{
+			return false;
+		}
+
 		return true;
 	}
 
 	virtual bool Initialize(TConstArrayView<NNE::FTensorDesc> Inputs, TConstArrayView<NNE::FTensorDesc> Outputs, const NNE::FAttributeMap& Attributes) override
 	{
-		check(Inputs.Num() >= 1 && Inputs.Num() <= 3);
-		check(Outputs.Num() == 1);
+		check(Inputs.Num() == NumAllowedInputTensors);
+		check(Outputs.Num() == NumAllowedOutputTensors);
 
 		// Read attributes
 		Epsilon = Attributes.GetValueOrDefault<float>(TEXT("epsilon"), DefaultEpsilon);
@@ -46,6 +81,17 @@ public:
 		const NNE::Internal::FTensor& InputTensor = *InputTensors[0];
 		const NNE::Internal::FTensor& OutputTensor = *OutputTensors[0];
 
+		// DML accepts only (N x C x H x W) shapes. When a smaller shape is provided we have to fill with 1s after N and C.
+		TArray<uint32> InputShape(InputTensor.GetShape().GetData());
+		{
+			const int32 InputShapeOffset = 4 - InputShape.Num();
+
+			for (int32 Idx = 0; Idx < InputShapeOffset; ++Idx)
+			{
+				InputShape.Insert(1, 2);
+			}
+		}
+
 		// Initialize tensor descriptors
 		FTensorDescDml	DmlInputTensorDesc;
 		FTensorDescDml	DmlScalingTensorDesc;
@@ -54,8 +100,9 @@ public:
 			
 		// Make sure that input is 4D
 		if (!DmlInputTensorDesc
-				.SetTensorRank(4, 4)
+				.SetTensorRank(MaxTensorRank, MaxTensorRank)
 				.SetFromTensor(InputTensor)
+				.SetShape(InputShape)
 				.Validate())
 		{
 			UE_LOG(LogNNE, Error, TEXT("Failed to initialize tensor(s) for DML inference"));
@@ -67,7 +114,7 @@ public:
 			const NNE::Internal::FTensor& ScaleTensor = *InputTensors[1];
 
 			if (!DmlScalingTensorDesc
-					.SetTensorRank(4, 4)
+					.SetTensorRank(MaxTensorRank, MaxTensorRank)
 					.SetFromTensor1D(ScaleTensor, InputTensor.GetShape().Rank())
 					.Validate())
 			{
@@ -81,7 +128,7 @@ public:
 			const NNE::Internal::FTensor& BiasTensor = *InputTensors[2];
 
 			if (!DmlBiasTensorDesc
-					.SetTensorRank(4, 4)
+					.SetTensorRank(MaxTensorRank, MaxTensorRank)
 					.SetFromTensor1D(BiasTensor, InputTensor.GetShape().Rank())
 					.Validate())
 			{
@@ -91,7 +138,7 @@ public:
 		}
 
 		if (!DmlOutputTensorDesc
-					.SetTensorRank(4, 4)
+					.SetTensorRank(MaxTensorRank, MaxTensorRank)
 					.SetFromTensor(OutputTensor)
 					.Validate())
 		{
diff --git a/Engine/Plugins/Experimental/NNERuntimeRDG/Source/NNERuntimeRDG/Private/Dml/NNEDmlOperatorLRN.cpp b/Engine/Plugins/Experimental/NNERuntimeRDG/Source/NNERuntimeRDG/Private/Dml/NNEDmlOperatorLRN.cpp
--- a/Engine/Plugins/Experimental/NNERuntimeRDG/Source/NNERuntimeRDG/Private/Dml/NNEDmlOperatorLRN.cpp
+++ b/Engine/Plugins/Experimental/NNERuntimeRDG/Source/NNERuntimeRDG/Private/Dml/NNEDmlOperatorLRN.cpp
@@ -13,6 +13,8 @@ class FOperatorDmlLRN : public FOperatorDml
 	float Alpha;
 	float Beta;
 	float Bias;
+	static constexpr uint32 NumAllowedInputTensors = 1, NumAllowedOutputTensors = 1;
+	static constexpr int32 	MinTensorRank = 4, MaxTensorRank = 4;
 
 public:
 
@@ -23,13 +25,30 @@ public:
 
 	static bool Validate(const NNE::FAttributeMap& AttributeMap, TConstArrayView<ENNETensorDataType> InputTypes, TConstArrayView<NNE::FSymbolicTensorShape> InputShapes)
 	{
+		const FString OpName = TEXT("LRN");
+
+		if(InputShapes.Num() != NumAllowedInputTensors)
+		{
+			UE_LOG(LogNNE, Warning, TEXT("DML %s: Invalid number of input tensors. %d provided, it should be %d."), *OpName, InputShapes.Num(), NumAllowedInputTensors);
+			return false;
+		}
+		
+		if (!CheckGenericTensor(OpName, InputTypes[0], InputShapes[0], 
+			{ 	ENNETensorDataType::Float, ENNETensorDataType::Half
+			},
+			MinTensorRank, MaxTensorRank
+		  	))
+		{
+			return false;
+		}
+
 		return true;
 	}
 
 	virtual bool Initialize(TConstArrayView<NNE::FTensorDesc> Inputs, TConstArrayView<NNE::FTensorDesc> Outputs, const NNE::FAttributeMap& Attributes) override
 	{
-		check(Inputs.Num() == 1);
-		check(Outputs.Num() == 1);
+		check(Inputs.Num() == NumAllowedInputTensors);
+		check(Outputs.Num() == NumAllowedOutputTensors);
 
 		// Read attributes
 		Size = Attributes.GetValueOrDefault<int32>(TEXT("size"), 0);
@@ -54,7 +73,7 @@ public:
 		FTensorDescDml	DmlInputTensorDesc;
 
 		if (!DmlInputTensorDesc
-				.SetTensorRank(4, 4)
+				.SetTensorRank(MinTensorRank, MaxTensorRank)
 				.SetFromTensor(InputTensor)
 				.Validate())
 		{
@@ -65,7 +84,7 @@ public:
 		FTensorDescDml	DmlOutputTensorDesc;
 
 		if (!DmlOutputTensorDesc
-				.SetTensorRank(4, 4)
+				.SetTensorRank(MinTensorRank, MaxTensorRank)
 				.SetFromTensor(OutputTensor)
 				.Validate())
 		{
diff --git a/Engine/Plugins/Experimental/NNERuntimeRDG/Source/NNERuntimeRDG/Private/Dml/NNEDmlOperatorLpNormalization.cpp b/Engine/Plugins/Experimental/NNERuntimeRDG/Source/NNERuntimeRDG/Private/Dml/NNEDmlOperatorLpNormalization.cpp
--- a/Engine/Plugins/Experimental/NNERuntimeRDG/Source/NNERuntimeRDG/Private/Dml/NNEDmlOperatorLpNormalization.cpp
+++ b/Engine/Plugins/Experimental/NNERuntimeRDG/Source/NNERuntimeRDG/Private/Dml/NNEDmlOperatorLpNormalization.cpp
@@ -14,6 +14,8 @@ class FOperatorDmlLpNormalization : public FOperatorDml
 	int32	P;
 	uint32	Axis;
 
+	static constexpr uint32 NumAllowedInputTensors = 1, NumAllowedOutputTensors = 1;
+	static constexpr int32 	MinTensorRank = 0, MaxTensorRank = GMaxTensorRank;
 
 public:
 
@@ -24,13 +26,30 @@ public:
 
 	static bool Validate(const NNE::FAttributeMap& AttributeMap, TConstArrayView<ENNETensorDataType> InputTypes, TConstArrayView<NNE::FSymbolicTensorShape> InputShapes)
 	{
+		const FString OpName = TEXT("LpNormalization");
+
+		if(InputShapes.Num() != NumAllowedInputTensors)
+		{
+			UE_LOG(LogNNE, Warning, TEXT("DML %s: Invalid number of input tensors. %d provided, it should be %d."), *OpName, InputShapes.Num(), NumAllowedInputTensors);
+			return false;
+		}
+		
+		if (!CheckGenericTensor(OpName, InputTypes[0], InputShapes[0], 
+			{ 	ENNETensorDataType::Float, ENNETensorDataType::Half
+			},
+			MinTensorRank, MaxTensorRank
+		  	))
+		{
+			return false;
+		}
+
 		return true;
 	}
 
 	virtual bool Initialize(TConstArrayView<NNE::FTensorDesc> Inputs, TConstArrayView<NNE::FTensorDesc> Outputs, const NNE::FAttributeMap& Attributes) override
 	{
-		check(Inputs.Num() == 1);
-		check(Outputs.Num() == 1);
+		check(Inputs.Num() == NumAllowedInputTensors);
+		check(Outputs.Num() == NumAllowedOutputTensors);
 
 		// Read attributes
 		P = Attributes.GetValueOrDefault<int32>(TEXT("p"), 1);
diff --git a/Engine/Plugins/Experimental/NNERuntimeRDG/Source/NNERuntimeRDG/Private/Dml/NNEDmlOperatorMatMul.cpp b/Engine/Plugins/Experimental/NNERuntimeRDG/Source/NNERuntimeRDG/Private/Dml/NNEDmlOperatorMatMul.cpp
--- a/Engine/Plugins/Experimental/NNERuntimeRDG/Source/NNERuntimeRDG/Private/Dml/NNEDmlOperatorMatMul.cpp
+++ b/Engine/Plugins/Experimental/NNERuntimeRDG/Source/NNERuntimeRDG/Private/Dml/NNEDmlOperatorMatMul.cpp
@@ -2,12 +2,15 @@
 
 #ifdef NNE_USE_DIRECTML
 #include "NNEDmlOperator.h"
+#include "NNEDmlOperatorUtils.h"
 
 namespace UE::NNERuntimeRDG::Private::Dml
 {
 
 class FOperatorDmlMatMul : public FOperatorDml
 {
+	static constexpr uint32 NumAllowedInputTensors = 2, NumAllowedOutputTensors = 1;
+	static constexpr int32 	MinTensorRank = 2, MaxTensorRank = TNumericLimits<int32>::Max();
 public:
 
 	static FOperatorDml* Create()
@@ -17,21 +20,47 @@ public:
 
 	static bool Validate(const NNE::FAttributeMap& AttributeMap, TConstArrayView<ENNETensorDataType> InputTypes, TConstArrayView<NNE::FSymbolicTensorShape> InputShapes)
 	{
+		const FString OpName = TEXT("MatMul");
+
+		if(InputShapes.Num() != NumAllowedInputTensors)
+		{
+			UE_LOG(LogNNE, Warning, TEXT("DML %s: Invalid number of input tensors. %d provided, it should be %d."), *OpName, InputShapes.Num(), NumAllowedInputTensors);
+			return false;
+		}
+		
+		if (!CheckGenericTensor(OpName, InputTypes[0], InputShapes[0], 
+			{ 	ENNETensorDataType::Float, ENNETensorDataType::Half
+			},
+			MinTensorRank, MaxTensorRank
+		  	))
+		{
+			return false;
+		}
+
+		if (!CheckGenericTensor(OpName, InputTypes[1], InputShapes[1], 
+			{ 	ENNETensorDataType::Float, ENNETensorDataType::Half
+			},
+			MinTensorRank, MaxTensorRank
+		  	))
+		{
+			return false;
+		}
+
 		return true;
 	}
 
 	virtual bool Initialize(TConstArrayView<NNE::FTensorDesc> Inputs, TConstArrayView<NNE::FTensorDesc> Outputs, const NNE::FAttributeMap& Attributes) override
 	{
-		check(Inputs.Num() == 2);
-		check(Outputs.Num() == 1);
+		check(Inputs.Num() == NumAllowedInputTensors);
+		check(Outputs.Num() == NumAllowedOutputTensors);
 
 		return true;
 	}
 
 	virtual int PrepareOutputs(TConstArrayView<NNE::Internal::FTensorRef> InputTensors, TArrayView<NNE::Internal::FTensorRef> OutputTensors) const override
 	{
-		check(InputTensors.Num() >= 2 && InputTensors.Num() <= 3);
-		check(OutputTensors.Num() == 1);
+		check(InputTensors.Num() == NumAllowedInputTensors);
+		check(OutputTensors.Num() == NumAllowedOutputTensors);
 
 		const NNE::FTensorShape& InputA = InputTensors[0]->GetShape();
 		const NNE::FTensorShape& InputB = InputTensors[1]->GetShape();
diff --git a/Engine/Plugins/Experimental/NNERuntimeRDG/Source/NNERuntimeRDG/Private/Dml/NNEDmlOperatorMaxUnpool.cpp b/Engine/Plugins/Experimental/NNERuntimeRDG/Source/NNERuntimeRDG/Private/Dml/NNEDmlOperatorMaxUnpool.cpp
--- a/Engine/Plugins/Experimental/NNERuntimeRDG/Source/NNERuntimeRDG/Private/Dml/NNEDmlOperatorMaxUnpool.cpp
+++ b/Engine/Plugins/Experimental/NNERuntimeRDG/Source/NNERuntimeRDG/Private/Dml/NNEDmlOperatorMaxUnpool.cpp
@@ -47,6 +47,8 @@ class FOperatorDmlMaxUnpool : public FOperatorDml
 	};
 
 	mutable FUnpoolingArgs Args;
+	static constexpr uint32 MinAllowedInputTensors = 2, MaxAllowedInputTensors = 3, NumAllowedOutputTensors = 1;
+	static constexpr int32 	MinTensorRank = 4, MaxTensorRank = 4;
 
 public:
 
@@ -57,13 +59,54 @@ public:
 
 	static bool Validate(const NNE::FAttributeMap& AttributeMap, TConstArrayView<ENNETensorDataType> InputTypes, TConstArrayView<NNE::FSymbolicTensorShape> InputShapes)
 	{
+		const FString OpName = TEXT("MaxUnpool");
+
+		if (InputShapes.Num() < MinAllowedInputTensors || InputShapes.Num() > MaxAllowedInputTensors)
+		{
+			UE_LOG(LogNNE, Warning, TEXT("DML %s: invalid number of input tensors. %d provided, it should be in [%d, %d]."), 
+										*OpName, InputShapes.Num(), MinAllowedInputTensors, MaxAllowedInputTensors);
+			return false;
+		}
+		
+		if (!CheckGenericTensor(OpName, InputTypes[0], InputShapes[0], 
+			{ ENNETensorDataType::Float, ENNETensorDataType::Half, 
+			  ENNETensorDataType::Int64, ENNETensorDataType::Int32, 
+			  ENNETensorDataType::Int16, ENNETensorDataType::Int8, 
+			  ENNETensorDataType::UInt64, ENNETensorDataType::UInt32,
+			  ENNETensorDataType::UInt16, ENNETensorDataType::UInt8 },
+			MinTensorRank, MaxTensorRank
+		  	))
+		{
+			return false;
+		}
+
+		if (!CheckGenericTensor(OpName, InputTypes[1], InputShapes[1], 
+			{ 	ENNETensorDataType::Int64
+			},
+			MinTensorRank, MaxTensorRank
+		  	))
+		{
+			return false;
+		}
+
+		if(InputShapes.Num() == 3)
+		{
+			if (!CheckGenericTensor1D(OpName, InputTypes[2], InputShapes[2], 
+				{ 	ENNETensorDataType::Int64
+				}
+				))
+			{
+				return false;
+			}
+		}
+
 		return true;
 	}
 
 	virtual bool Initialize(TConstArrayView<NNE::FTensorDesc> Inputs, TConstArrayView<NNE::FTensorDesc> Outputs, const NNE::FAttributeMap& Attributes) override
 	{
-		check(Inputs.Num() == 2 || Inputs.Num() == 3);
-		check(Outputs.Num() == 1);
+		check(Inputs.Num() >= MinAllowedInputTensors && Inputs.Num() <= MaxAllowedInputTensors);
+		check(Outputs.Num() == NumAllowedOutputTensors);
 
 		if (Inputs.Num() == 3)
 		{
@@ -126,7 +169,7 @@ public:
 		FTensorDescDml	DmlOutputTensorDesc;
 
 		if (!DmlInputTensorDesc
-				.SetTensorRank(4, 4)
+				.SetTensorRank(MinTensorRank, MaxTensorRank)
 				.SetFromTensor(InputTensor)
 				.Validate())
 		{
@@ -135,7 +178,7 @@ public:
 		}
 
 		if (!DmlIndicesTensorDesc
-				.SetTensorRank(4, 4)
+				.SetTensorRank(MinTensorRank, MaxTensorRank)
 				.SetFromTensor(IndicesTensor)
 				.SetDataType(ENNETensorDataType::UInt64)
 				.Validate())
@@ -145,7 +188,7 @@ public:
 		}
 		
 		if (!DmlOutputTensorDesc
-				.SetTensorRank(4, 4)
+				.SetTensorRank(MinTensorRank, MaxTensorRank)
 				.SetFromTensor(OutputTensor)
 				.Validate())
 		{
diff --git a/Engine/Plugins/Experimental/NNERuntimeRDG/Source/NNERuntimeRDG/Private/Dml/NNEDmlOperatorMeanVarianceNormalization.cpp b/Engine/Plugins/Experimental/NNERuntimeRDG/Source/NNERuntimeRDG/Private/Dml/NNEDmlOperatorMeanVarianceNormalization.cpp
--- a/Engine/Plugins/Experimental/NNERuntimeRDG/Source/NNERuntimeRDG/Private/Dml/NNEDmlOperatorMeanVarianceNormalization.cpp
+++ b/Engine/Plugins/Experimental/NNERuntimeRDG/Source/NNERuntimeRDG/Private/Dml/NNEDmlOperatorMeanVarianceNormalization.cpp
@@ -14,6 +14,8 @@ class FOperatorDmlMeanVarianceNormalization : public FOperatorDml
 	Util::FSmallUIntArray	Axes;
 	int32					AcrossChannels;
 	int32					NormalizeVariance;
+	static constexpr uint32 NumAllowedInputTensors = 1, NumAllowedOutputTensors = 1;
+	static constexpr int32 	MinTensorRank = 0, MaxTensorRank = GMaxTensorRank;
 
 public:
 
@@ -24,13 +26,30 @@ public:
 
 	static bool Validate(const NNE::FAttributeMap& AttributeMap, TConstArrayView<ENNETensorDataType> InputTypes, TConstArrayView<NNE::FSymbolicTensorShape> InputShapes)
 	{
+		const FString OpName = TEXT("MeanVarianceNormalization");
+
+		if(InputShapes.Num() != NumAllowedInputTensors)
+		{
+			UE_LOG(LogNNE, Warning, TEXT("DML %s: Invalid number of input tensors. %d provided, it should be %d."), *OpName, InputShapes.Num(), NumAllowedInputTensors);
+			return false;
+		}
+		
+		if (!CheckGenericTensor(OpName, InputTypes[0], InputShapes[0], 
+			{ 	ENNETensorDataType::Float, ENNETensorDataType::Half
+			},
+			MinTensorRank, MaxTensorRank
+		  	))
+		{
+			return false;
+		}
+
 		return true;
 	}
 
 	virtual bool Initialize(TConstArrayView<NNE::FTensorDesc> Inputs, TConstArrayView<NNE::FTensorDesc> Outputs, const NNE::FAttributeMap& Attributes) override
 	{
-		check(Inputs.Num() == 1);
-		check(Outputs.Num() == 1);
+		check(Inputs.Num() == NumAllowedInputTensors);
+		check(Outputs.Num() == NumAllowedOutputTensors);
 
 		// Read attributes
 		AcrossChannels = Attributes.GetValueOrDefault<int32>(TEXT("across_channels"), 0);
@@ -75,12 +94,6 @@ public:
 		const NNE::Internal::FTensor& InputTensor = *InputTensors[0];
 		const NNE::Internal::FTensor& OutputTensor = *OutputTensors[0];
 
-		if (InputTensor.GetShape().Rank() > 8)
-		{
-			UE_LOG(LogNNE, Warning, TEXT("InputTensor rank should be between 1 and 8, got:%d"), InputTensor.GetShape().Rank());
-			return false;
-		}
-
 		FTensorDescDml	DmlInputTensorDesc;
 		
 		if (!DmlInputTensorDesc
diff --git a/Engine/Plugins/Experimental/NNERuntimeRDG/Source/NNERuntimeRDG/Private/Dml/NNEDmlOperatorPad.cpp b/Engine/Plugins/Experimental/NNERuntimeRDG/Source/NNERuntimeRDG/Private/Dml/NNEDmlOperatorPad.cpp
--- a/Engine/Plugins/Experimental/NNERuntimeRDG/Source/NNERuntimeRDG/Private/Dml/NNEDmlOperatorPad.cpp
+++ b/Engine/Plugins/Experimental/NNERuntimeRDG/Source/NNERuntimeRDG/Private/Dml/NNEDmlOperatorPad.cpp
@@ -2,6 +2,7 @@
 
 #ifdef NNE_USE_DIRECTML
 #include "NNEDmlOperator.h"
+#include "NNEDmlOperatorUtils.h"
 
 namespace UE::NNERuntimeRDG::Private::Dml
 {
@@ -32,6 +33,9 @@ class FOperatorDmlPad : public FOperatorDml
 		}
 	}
 
+	static constexpr uint32 MinAllowedInputTensors = 1, MaxAllowedInputTensors = 2, NumAllowedOutputTensors = 1;
+	static constexpr int32 	MinTensorRank = 0, MaxTensorRank = GMaxTensorRank;
+
 public:
 
 	static FOperatorDml* Create()
@@ -41,13 +45,47 @@ public:
 
 	static bool Validate(const NNE::FAttributeMap& AttributeMap, TConstArrayView<ENNETensorDataType> InputTypes, TConstArrayView<NNE::FSymbolicTensorShape> InputShapes)
 	{
+		const FString OpName = TEXT("Pad");
+
+		if (InputShapes.Num() < MinAllowedInputTensors || InputShapes.Num() > MaxAllowedInputTensors)
+		{
+			UE_LOG(LogNNE, Warning, TEXT("DML %s: invalid number of input tensors. %d provided, it should be in [%d, %d]."), 
+										*OpName, InputShapes.Num(), MinAllowedInputTensors, MaxAllowedInputTensors);
+			return false;
+		}
+
+		if (!CheckGenericTensor(OpName, InputTypes[0], InputShapes[0], 
+			{ 	ENNETensorDataType::Double, ENNETensorDataType::Float, ENNETensorDataType::Half, 
+				ENNETensorDataType::Int64, ENNETensorDataType::Int32, ENNETensorDataType::Int16,
+				ENNETensorDataType::Int8, ENNETensorDataType::UInt64, ENNETensorDataType::UInt32, 
+				ENNETensorDataType::UInt16, ENNETensorDataType::UInt8
+			},
+			MinTensorRank, MaxTensorRank
+		  	))
+		{
+			return false;
+		}
+
+		if(InputShapes.Num() >= 2)
+		{
+			//pads are 1D
+			if (!CheckGenericTensor1D(OpName, InputTypes[1], InputShapes[1], 
+				{ 	
+					ENNETensorDataType::Int64
+				}
+				))
+			{
+				return false;
+			}
+		}
+
 		return true;
 	}
 
 	virtual bool Initialize(TConstArrayView<NNE::FTensorDesc> Inputs, TConstArrayView<NNE::FTensorDesc> Outputs, const NNE::FAttributeMap& Attributes) override
 	{
-		check(Inputs.Num() >= 1);
-		check(Outputs.Num() == 1);
+		check(Inputs.Num() >= MinAllowedInputTensors && Inputs.Num() <= MaxAllowedInputTensors);
+		check(Outputs.Num() == NumAllowedOutputTensors);
 
 		// Read attributes
 		Value = Attributes.GetValueOrDefault<float>(TEXT("value"), 0.0f);
diff --git a/Engine/Plugins/Experimental/NNERuntimeRDG/Source/NNERuntimeRDG/Private/Dml/NNEDmlOperatorPool.cpp b/Engine/Plugins/Experimental/NNERuntimeRDG/Source/NNERuntimeRDG/Private/Dml/NNEDmlOperatorPool.cpp
--- a/Engine/Plugins/Experimental/NNERuntimeRDG/Source/NNERuntimeRDG/Private/Dml/NNEDmlOperatorPool.cpp
+++ b/Engine/Plugins/Experimental/NNERuntimeRDG/Source/NNERuntimeRDG/Private/Dml/NNEDmlOperatorPool.cpp
@@ -13,7 +13,7 @@ namespace UE::NNERuntimeRDG::Private::Dml
 /**
  *  Pooling operators (local and global)
  */
-template <typename TDmlPoolOpDesc, DML_OPERATOR_TYPE DmlOpType, bool UseGlobalPooling>
+template <typename TDmlPoolOpDesc, DML_OPERATOR_TYPE DmlOpType, bool UseGlobalPooling, TCHAR const *OpName>
 class FOperatorDmlPool : public FOperatorDml
 {
 	using FIntArray = TArray<int32>;
@@ -79,24 +79,61 @@ class FOperatorDmlPool : public FOperatorDml
 	mutable FPoolingArgs	Args;
 	int32					IncludePadding;
 	uint32					P;
+	static constexpr uint32 NumAllowedInputTensors = 1, MinAllowedOutputTensors = 1, MaxAllowedOutputTensors = 2;
+	static constexpr int32 	MinTensorRank = 4, MaxTensorRank = 5;
 
 public:
 
 	static FOperatorDml* Create()
 	{
-		return new FOperatorDmlPool<TDmlPoolOpDesc, DmlOpType, UseGlobalPooling>();
+		return new FOperatorDmlPool<TDmlPoolOpDesc, DmlOpType, UseGlobalPooling, OpName>();
 	}
 
 	static bool Validate(const NNE::FAttributeMap& AttributeMap, TConstArrayView<ENNETensorDataType> InputTypes, TConstArrayView<NNE::FSymbolicTensorShape> InputShapes)
 	{
+
+		if(InputShapes.Num() != NumAllowedInputTensors)
+		{
+			UE_LOG(LogNNE, Warning, TEXT("DML %s: Invalid number of input tensors. %d provided, it should be %d."), OpName, InputShapes.Num(), NumAllowedInputTensors);
+			return false;
+		}
+
+		TSet<ENNETensorDataType> AllowedDataTypes;
+		if constexpr (DmlOpType == DML_OPERATOR_MAX_POOLING)
+		{
+			AllowedDataTypes = 
+					{ 	
+						ENNETensorDataType::Float, ENNETensorDataType::Half, 
+						ENNETensorDataType::Int64, ENNETensorDataType::Int32, ENNETensorDataType::Int16,
+						ENNETensorDataType::Int8, ENNETensorDataType::UInt64, ENNETensorDataType::UInt32, 
+						ENNETensorDataType::UInt16, ENNETensorDataType::UInt8
+					};
+		}
+		else
+		{
+			AllowedDataTypes = 
+					{ 	
+						ENNETensorDataType::Float, ENNETensorDataType::Half
+					};
+		}
+
+		if (!CheckGenericTensor(OpName, InputTypes[0], InputShapes[0], 
+			AllowedDataTypes,
+			MinTensorRank, MaxTensorRank
+			))
+		{
+			return false;
+		}
+
+
 		return true;
 	}
 
 	virtual bool Initialize(TConstArrayView<NNE::FTensorDesc> Inputs, TConstArrayView<NNE::FTensorDesc> Outputs, const NNE::FAttributeMap& Attributes) override
 	{
 		// TODO: int64 attributes
-		check(Inputs.Num() == 1);
-		check(Outputs.Num() == 1 || Outputs.Num() == 2);
+		check(Inputs.Num() == NumAllowedInputTensors);
+		check(Outputs.Num() >= MinAllowedOutputTensors && Outputs.Num() <= MaxAllowedOutputTensors);
 
 		// Attribute storage_order is not supported
 		const int32 StorageOrder = Attributes.GetValueOrDefault<int32>(TEXT("storage_order"), 0);
@@ -173,7 +210,7 @@ public:
 		FTensorDescDml DmlInputTensorDesc;
 
 		if (!DmlInputTensorDesc
-				.SetTensorRank(4, 5)
+				.SetTensorRank(MinTensorRank, MaxTensorRank)
 				.SetFromTensor(*InputTensors[0])
 				.Validate())
 		{
@@ -184,7 +221,7 @@ public:
 		FTensorDescDml DmlOutputTensorDesc;
 
 		if (!DmlOutputTensorDesc
-				.SetTensorRank(4, 5)
+				.SetTensorRank(MinTensorRank, MaxTensorRank)
 				.SetFromTensor(*OutputTensors[0])
 				.Validate())
 		{
@@ -225,7 +262,7 @@ public:
 				FTensorDescDml DmlIndicesTensorDesc;
 				
 				if (!DmlIndicesTensorDesc
-						.SetTensorRank(4, 5)
+						.SetTensorRank(MinTensorRank, MaxTensorRank)
 						.SetFromTensor(*OutputTensors[1])
 						.SetDataType(ENNETensorDataType::UInt64)
 						.Validate())
@@ -262,11 +299,12 @@ public:
 };
 
 #define NNE_DML_REGISTER_POOLING_OP(OpName, DmlPrefix, UseGlobalPooling) \
+TCHAR const Op##OpName##Name[] = TEXT(#OpName); \
 struct FDmlOperator##OpName##Registrator \
 { \
 	FDmlOperator##OpName##Registrator() \
 	{ \
-		FOperatorRegistryDml::Get()->OpAdd(TEXT(#OpName), FOperatorDmlPool<DML_##DmlPrefix##_POOLING_OPERATOR_DESC, DML_OPERATOR_##DmlPrefix##_POOLING, UseGlobalPooling>::Create); \
+		FOperatorRegistryDml::Get()->OpAdd(TEXT(#OpName), FOperatorDmlPool<DML_##DmlPrefix##_POOLING_OPERATOR_DESC, DML_OPERATOR_##DmlPrefix##_POOLING, UseGlobalPooling, Op##OpName##Name>::Create, FOperatorDmlPool<DML_##DmlPrefix##_POOLING_OPERATOR_DESC, DML_OPERATOR_##DmlPrefix##_POOLING, UseGlobalPooling, Op##OpName##Name>::Validate); \
 	} \
 }; \
 \
diff --git a/Engine/Plugins/Experimental/NNERuntimeRDG/Source/NNERuntimeRDG/Private/Dml/NNEDmlOperatorReduce.cpp b/Engine/Plugins/Experimental/NNERuntimeRDG/Source/NNERuntimeRDG/Private/Dml/NNEDmlOperatorReduce.cpp
--- a/Engine/Plugins/Experimental/NNERuntimeRDG/Source/NNERuntimeRDG/Private/Dml/NNEDmlOperatorReduce.cpp
+++ b/Engine/Plugins/Experimental/NNERuntimeRDG/Source/NNERuntimeRDG/Private/Dml/NNEDmlOperatorReduce.cpp
@@ -19,7 +19,7 @@ namespace UE::NNERuntimeRDG::Private::Dml
 // ReduceProd
 // ReduceSum
 // ReduceSumSquare
-template<DML_REDUCE_FUNCTION ReduceFunc>
+template<DML_REDUCE_FUNCTION ReduceFunc, TCHAR const *OpName>
 class FOperatorDmlReduce : public FOperatorDml
 {
 	inline static void HandleEmptyAxes(TArray<int32>& Axes, int32 Rank)
@@ -40,23 +40,43 @@ class FOperatorDmlReduce : public FOperatorDml
 	mutable Util::FSmallArray<bool>	IsReducedDims;
 	int32							KeepDims;
 	DML_AXIS_DIRECTION				AxisDirection{ DML_AXIS_DIRECTION_INCREASING };
+	static constexpr uint32 NumAllowedInputTensors = 1, NumAllowedOutputTensors = 1;
+	static constexpr int32 	MinTensorRank = 0, MaxTensorRank = GMaxTensorRank;
 
 public:
 
 	static FOperatorDml* Create()
 	{
-		return new FOperatorDmlReduce<ReduceFunc>();
+		return new FOperatorDmlReduce<ReduceFunc, OpName>();
 	}
 
 	static bool Validate(const NNE::FAttributeMap& AttributeMap, TConstArrayView<ENNETensorDataType> InputTypes, TConstArrayView<NNE::FSymbolicTensorShape> InputShapes)
 	{
+		if(InputShapes.Num() != NumAllowedInputTensors)
+		{
+			UE_LOG(LogNNE, Warning, TEXT("DML %s: Invalid number of input tensors. %d provided, it should be %d."), OpName, InputShapes.Num(), NumAllowedInputTensors);
+			return false;
+		}
+		
+		if (!CheckGenericTensor(OpName, InputTypes[0], InputShapes[0], 
+			{ 	ENNETensorDataType::Float, ENNETensorDataType::Half, 
+				ENNETensorDataType::Int64, ENNETensorDataType::Int32, ENNETensorDataType::Int16,
+				ENNETensorDataType::Int8, ENNETensorDataType::UInt64, ENNETensorDataType::UInt32, 
+				ENNETensorDataType::UInt16, ENNETensorDataType::UInt8
+			},
+			MinTensorRank, MaxTensorRank
+			))
+		{
+			return false;
+		}
+
 		return true;
 	}
 
 	virtual bool Initialize(TConstArrayView<NNE::FTensorDesc> Inputs, TConstArrayView<NNE::FTensorDesc> Outputs, const NNE::FAttributeMap& Attributes) override
 	{
-		checkf(Inputs.Num() == 1, TEXT("Dml Reduce op supports only 1 input"));
-		check(Outputs.Num() == 1);
+		checkf(Inputs.Num() == NumAllowedInputTensors, TEXT("Dml Reduce op supports only 1 input"));
+		check(Outputs.Num() == NumAllowedOutputTensors);
 
 		KeepDims = Attributes.GetValueOrDefault<int32>(TEXT("keepdims"), 1);
 
@@ -217,31 +237,35 @@ public:
 	}
 };
 
-// Register Reshape operator on Module startup
-#define OP(OpName, ReduceFunc) FOperatorRegistryDml::Get()->OpAdd(TEXT(#OpName), FOperatorDmlReduce<ReduceFunc>::Create)
-
-struct FOperatorDmlReduceRegistrator
-{
-	FOperatorDmlReduceRegistrator()
-	{
-		OP(ReduceL1,		DML_REDUCE_FUNCTION_L1);
-		OP(ReduceL2,		DML_REDUCE_FUNCTION_L2);
-		OP(ReduceLogSum,	DML_REDUCE_FUNCTION_LOG_SUM);
-		OP(ReduceLogSumExp,	DML_REDUCE_FUNCTION_LOG_SUM_EXP);
-		OP(ReduceMin,		DML_REDUCE_FUNCTION_MIN);
-		OP(ReduceMax,		DML_REDUCE_FUNCTION_MAX);
-		OP(ReduceMean,		DML_REDUCE_FUNCTION_AVERAGE);
-		OP(ReduceProd,		DML_REDUCE_FUNCTION_MULTIPLY);
-		OP(ReduceSum,		DML_REDUCE_FUNCTION_SUM);
-		OP(ReduceSumSquare,	DML_REDUCE_FUNCTION_SUM_SQUARE);
-		OP(ArgMax,			DML_REDUCE_FUNCTION_ARGMAX);
-		OP(ArgMin,			DML_REDUCE_FUNCTION_ARGMIN);
-	}
-};
 
-#undef OP
+#define NNE_DML_REGISTER_REDUCE_OP(OpName, ReduceFunc) \
+TCHAR const Op##OpName##Name[] = TEXT(#OpName); \
+struct FDmlOperator##OpName##Registrator \
+{ \
+	FDmlOperator##OpName##Registrator() \
+	{ \
+		FOperatorRegistryDml::Get()->OpAdd(TEXT(#OpName), FOperatorDmlReduce<ReduceFunc, Op##OpName##Name>::Create, FOperatorDmlReduce<ReduceFunc, Op##OpName##Name>::Validate); \
+	} \
+}; \
+\
+static FDmlOperator##OpName##Registrator RegisterDmlOperator##OpName;
+
+NNE_DML_REGISTER_REDUCE_OP(ReduceL1,		DML_REDUCE_FUNCTION_L1)
+NNE_DML_REGISTER_REDUCE_OP(ReduceL2,		DML_REDUCE_FUNCTION_L2)
+NNE_DML_REGISTER_REDUCE_OP(ReduceLogSum,	DML_REDUCE_FUNCTION_LOG_SUM)
+NNE_DML_REGISTER_REDUCE_OP(ReduceLogSumExp,	DML_REDUCE_FUNCTION_LOG_SUM_EXP)
+NNE_DML_REGISTER_REDUCE_OP(ReduceMin,		DML_REDUCE_FUNCTION_MIN)
+NNE_DML_REGISTER_REDUCE_OP(ReduceMax,		DML_REDUCE_FUNCTION_MAX)
+NNE_DML_REGISTER_REDUCE_OP(ReduceMean,		DML_REDUCE_FUNCTION_AVERAGE)
+NNE_DML_REGISTER_REDUCE_OP(ReduceProd,		DML_REDUCE_FUNCTION_MULTIPLY)
+NNE_DML_REGISTER_REDUCE_OP(ReduceSum,		DML_REDUCE_FUNCTION_SUM)
+NNE_DML_REGISTER_REDUCE_OP(ReduceSumSquare,	DML_REDUCE_FUNCTION_SUM_SQUARE)
+NNE_DML_REGISTER_REDUCE_OP(ArgMax,			DML_REDUCE_FUNCTION_ARGMAX)
+NNE_DML_REGISTER_REDUCE_OP(ArgMin,			DML_REDUCE_FUNCTION_ARGMIN)
+
+
+#undef NNE_DML_REGISTER_REDUCE_OP
 
-static FOperatorDmlReduceRegistrator RegisterReduceOperators;
 
 
 } // namespace UE::NNERuntimeRDG::Private::Dml
diff --git a/Engine/Plugins/Experimental/NNERuntimeRDG/Source/NNERuntimeRDG/Private/Dml/NNEDmlOperatorResample.cpp b/Engine/Plugins/Experimental/NNERuntimeRDG/Source/NNERuntimeRDG/Private/Dml/NNEDmlOperatorResample.cpp
--- a/Engine/Plugins/Experimental/NNERuntimeRDG/Source/NNERuntimeRDG/Private/Dml/NNEDmlOperatorResample.cpp
+++ b/Engine/Plugins/Experimental/NNERuntimeRDG/Source/NNERuntimeRDG/Private/Dml/NNEDmlOperatorResample.cpp
@@ -2,6 +2,8 @@
 
 #ifdef NNE_USE_DIRECTML
 #include "NNEDmlOperator.h"
+#include "NNEDmlOperatorUtils.h"
+#include "Algo/Count.h"
 
 namespace UE::NNERuntimeRDG::Private::Dml
 {
@@ -63,6 +65,11 @@ class FOperatorDmlResample : public FOperatorDml
 		}
 	}
 
+	static constexpr uint32 ResizeMinAllowedInputTensors = 1, ResizeMaxAllowedInputTensors = 4;
+	static constexpr uint32 UpsampleNumAllowedInputTensors = 2;
+	static constexpr uint32 NumAllowedOutputTensors = 1;
+	static constexpr int32 	MinTensorRank = 0, MaxTensorRank = 4;
+
 public:
 
 	static FOperatorDml* Create()
@@ -72,13 +79,93 @@ public:
 
 	static bool Validate(const NNE::FAttributeMap& AttributeMap, TConstArrayView<ENNETensorDataType> InputTypes, TConstArrayView<NNE::FSymbolicTensorShape> InputShapes)
 	{
+		const FString OpName = IsResize ? TEXT("Resize") : TEXT("Upsample");
+		if constexpr (IsResize)
+		{
+			if (InputShapes.Num() < ResizeMinAllowedInputTensors || InputShapes.Num() > ResizeMaxAllowedInputTensors)
+			{
+				UE_LOG(LogNNE, Warning, TEXT("DML %s: invalid number of input tensors. %d provided, it should be in [%d, %d]."), 
+											*OpName, InputShapes.Num(), ResizeMinAllowedInputTensors, ResizeMaxAllowedInputTensors);
+				return false;
+			}
+
+			if(InputShapes.Num() > 1)
+			{
+				if (!CheckGenericTensor1D(OpName, InputTypes[1], InputShapes[1], 
+					{ 	ENNETensorDataType::Double, ENNETensorDataType::Float, ENNETensorDataType::Half
+					}
+					))
+				{
+					return false;
+				}
+			}
+
+			if(InputShapes.Num() == 3)
+			{
+				if (!CheckGenericTensor1D(OpName, InputTypes[2], InputShapes[2], 
+					{ 	ENNETensorDataType::Float
+					}
+					))
+				{
+					return false;
+				}
+			}
+
+			if(InputShapes.Num() == 4)
+			{
+				if (!CheckGenericTensor1D(OpName, InputTypes[3], InputShapes[3], 
+					{ 	ENNETensorDataType::Int64
+					}
+					))
+				{
+					return false;
+				}
+			}
+		}
+		else
+		{
+			if(InputShapes.Num() != UpsampleNumAllowedInputTensors)
+			{
+				UE_LOG(LogNNE, Warning, TEXT("DML %s: Invalid number of input tensors. %d provided, it should be %d."), *OpName, InputShapes.Num(), UpsampleNumAllowedInputTensors);
+				return false;
+			}
+
+			if(InputShapes.Num() > 1)
+			{
+				if (!CheckGenericTensor1D(OpName, InputTypes[1], InputShapes[1], 
+					{ 	ENNETensorDataType::Float
+					}
+					))
+				{
+					return false;
+				}
+			}
+		}
+		
+		if (!CheckGenericTensor(OpName, InputTypes[0], InputShapes[0], 
+			{ 	
+				ENNETensorDataType::Float, ENNETensorDataType::Half
+			},
+			MinTensorRank, MaxTensorRank + Algo::Count(InputShapes[0].GetData(), (int32) 1) // Allow ones due to squeezing (see Create method)
+		  	))
+		{
+			return false;
+		}
+
 		return true;
 	}
 
 	virtual bool Initialize(TConstArrayView<NNE::FTensorDesc> Inputs, TConstArrayView<NNE::FTensorDesc> Outputs, const NNE::FAttributeMap& Attributes) override
 	{
-		check(Inputs.Num() >= 1 && Inputs.Num() <= 4);
-		check(Outputs.Num() == 1);
+		if constexpr (IsResize)
+		{
+			check(Inputs.Num() >= ResizeMinAllowedInputTensors && Inputs.Num() <= ResizeMaxAllowedInputTensors);
+		}
+		else
+		{
+			check(Inputs.Num() == UpsampleNumAllowedInputTensors);
+		}
+		check(Outputs.Num() == NumAllowedOutputTensors);
 
 		bUseSizesTensor = Inputs.Num() == 4;
 
@@ -316,11 +403,11 @@ public:
 		FTensorDescDml	DmlOutputTensorDesc;
 
 		DmlInputTensorDesc
-			.SetTensorRank(1, 4)
+			.SetTensorRank(MinTensorRank, MaxTensorRank)
 			.SetFromTensor(InputTensor);
 		
 		DmlOutputTensorDesc
-			.SetTensorRank(1, 4)
+			.SetTensorRank(MinTensorRank, MaxTensorRank)
 			.SetFromTensor(OutputTensor);
 		
 		// Find any useless dimensions of size 1 that occur in both input and output
@@ -384,7 +471,7 @@ struct FDmlOperator##OpName##Registrator \
 { \
 	FDmlOperator##OpName##Registrator() \
 	{ \
-		FOperatorRegistryDml::Get()->OpAdd(TEXT(#OpName), FOperatorDmlResample<IsResize>::Create); \
+		FOperatorRegistryDml::Get()->OpAdd(TEXT(#OpName), FOperatorDmlResample<IsResize>::Create, FOperatorDmlResample<IsResize>::Validate); \
 	} \
 }; \
 \
diff --git a/Engine/Plugins/Experimental/NNERuntimeRDG/Source/NNERuntimeRDG/Private/Dml/NNEDmlOperatorReshape.cpp b/Engine/Plugins/Experimental/NNERuntimeRDG/Source/NNERuntimeRDG/Private/Dml/NNEDmlOperatorReshape.cpp
--- a/Engine/Plugins/Experimental/NNERuntimeRDG/Source/NNERuntimeRDG/Private/Dml/NNEDmlOperatorReshape.cpp
+++ b/Engine/Plugins/Experimental/NNERuntimeRDG/Source/NNERuntimeRDG/Private/Dml/NNEDmlOperatorReshape.cpp
@@ -3,6 +3,7 @@
 #ifdef NNE_USE_DIRECTML
 
 #include "NNEDmlOperator.h"
+#include "NNEDmlOperatorUtils.h"
 #include "Algo/Transform.h"
 #include "Algo/Count.h"
 #include "Misc/EnumerateRange.h"
@@ -15,6 +16,8 @@ class FOperatorDmlReshape : public FOperatorDml
 {
 	mutable Util::FSmallUIntArray	OutputShape;
 	bool							bAllowZero;
+	static constexpr uint32 NumAllowedInputTensors = 2, NumAllowedOutputTensors = 1;
+	static constexpr int32 	MinTensorRank = 0, MaxTensorRank = GMaxTensorRank;
 
 public:
 
@@ -25,13 +28,44 @@ public:
 
     static bool Validate(const NNE::FAttributeMap& AttributeMap, TConstArrayView<ENNETensorDataType> InputTypes, TConstArrayView<NNE::FSymbolicTensorShape> InputShapes)
 	{
+		const FString OpName = TEXT("Reshape");
+
+		if(InputShapes.Num() != NumAllowedInputTensors)
+		{
+			UE_LOG(LogNNE, Warning, TEXT("DML %s: Invalid number of input tensors. %d provided, it should be %d."), *OpName, InputShapes.Num(), NumAllowedInputTensors);
+			return false;
+		}
+		
+		if (!CheckGenericTensor(OpName, InputTypes[0], InputShapes[0], 
+			{ 	ENNETensorDataType::Double, ENNETensorDataType::Float, ENNETensorDataType::Half, 
+				ENNETensorDataType::Int64, ENNETensorDataType::Int32, ENNETensorDataType::Int16,
+				ENNETensorDataType::Int8, ENNETensorDataType::UInt64, ENNETensorDataType::UInt32, 
+				ENNETensorDataType::UInt16, ENNETensorDataType::UInt8
+			},
+			MinTensorRank, MaxTensorRank
+		  	))
+		{
+			return false;
+		}
+
+		if(InputShapes.Num() == 2)
+		{
+			if (!CheckGenericTensor1D(OpName, InputTypes[1], InputShapes[1], 
+				{ 	ENNETensorDataType::Int64, ENNETensorDataType::Int32, ENNETensorDataType::UInt32
+				}
+				))
+			{
+				return false;
+			}
+		}
+
 		return true;
 	}
 
 	virtual bool Initialize(TConstArrayView<NNE::FTensorDesc> Inputs, TConstArrayView<NNE::FTensorDesc> Outputs, const NNE::FAttributeMap& Attributes) override
 	{
-		check(Inputs.Num() == 2);
-		check(Outputs.Num() == 1);
+		check(Inputs.Num() == NumAllowedInputTensors);
+		check(Outputs.Num() == NumAllowedOutputTensors);
 
 		ConstantCPUInputs.Add(1);
 
diff --git a/Engine/Plugins/Experimental/NNERuntimeRDG/Source/NNERuntimeRDG/Private/Dml/NNEDmlOperatorSlice.cpp b/Engine/Plugins/Experimental/NNERuntimeRDG/Source/NNERuntimeRDG/Private/Dml/NNEDmlOperatorSlice.cpp
--- a/Engine/Plugins/Experimental/NNERuntimeRDG/Source/NNERuntimeRDG/Private/Dml/NNEDmlOperatorSlice.cpp
+++ b/Engine/Plugins/Experimental/NNERuntimeRDG/Source/NNERuntimeRDG/Private/Dml/NNEDmlOperatorSlice.cpp
@@ -3,6 +3,7 @@
 #ifdef NNE_USE_DIRECTML
 
 #include "NNEDmlOperator.h"
+#include "NNEDmlOperatorUtils.h"
 #include "Misc/EnumerateRange.h"
 #include "Algo/Copy.h"
 #include "Algo/ForEach.h"
@@ -133,6 +134,9 @@ class FOperatorDmlSlice : public FOperatorDml
 
 	ENNETensorDataType InputIndexDataType;
 
+	static constexpr uint32 MinAllowedInputTensors = 3, MaxAllowedInputTensors = 5, NumAllowedOutputTensors = 1;
+	static constexpr int32 	MinTensorRank = 0, MaxTensorRank = GMaxTensorRank;
+
 public:
 
 	static FOperatorDml* Create()
@@ -142,14 +146,63 @@ public:
 
 	static bool Validate(const NNE::FAttributeMap& AttributeMap, TConstArrayView<ENNETensorDataType> InputTypes, TConstArrayView<NNE::FSymbolicTensorShape> InputShapes)
 	{
+		const FString OpName = TEXT("Slice");
+
+		if (InputShapes.Num() < MinAllowedInputTensors || InputShapes.Num() > MaxAllowedInputTensors)
+		{
+			UE_LOG(LogNNE, Warning, TEXT("DML %s: invalid number of input tensors. %d provided, it should be in [%d, %d]."), 
+										*OpName, InputShapes.Num(), MinAllowedInputTensors, MaxAllowedInputTensors);
+			return false;
+		}
+		
+		if (!CheckGenericTensor(OpName, InputTypes[0], InputShapes[0], 
+			{ 	ENNETensorDataType::Double, ENNETensorDataType::Float, ENNETensorDataType::Half, 
+				ENNETensorDataType::Int64, ENNETensorDataType::Int32, ENNETensorDataType::Int16,
+				ENNETensorDataType::Int8, ENNETensorDataType::UInt64, ENNETensorDataType::UInt32, 
+				ENNETensorDataType::UInt16, ENNETensorDataType::UInt8
+			},
+			MinTensorRank, MaxTensorRank
+		  	))
+		{
+			return false;
+		}
+		
+		ENNETensorDataType IndexDataTypeToCheck = ENNETensorDataType::None;
+
+		for (int Idx = 1; Idx < InputShapes.Num(); Idx++)
+		{
+			if (!CheckGenericTensor1D(OpName, InputTypes[Idx], InputShapes[Idx], 
+				{ 	ENNETensorDataType::Int64, ENNETensorDataType::Int32
+				}
+				))
+			{
+				return false;
+			}
+			
+			if (Idx == 1)
+			{
+				IndexDataTypeToCheck = InputTypes[Idx];
+			}
+			else
+			{
+				if (InputTypes[Idx] != IndexDataTypeToCheck)
+				{
+					UE_LOG(LogNNE, Warning, TEXT("DML %s: data type of tensor at position %d differs from data type of 'starts' tensor."), 
+										*OpName, Idx);
+					return false;
+				}
+			}
+
+		}
+
 		return true;
 	}
 
 	virtual bool Initialize(TConstArrayView<NNE::FTensorDesc> Inputs, TConstArrayView<NNE::FTensorDesc> Outputs, const NNE::FAttributeMap& Attributes) override
 	{
-		check(Inputs.Num() >= 3);
-		check(Inputs.Num() <= 5);
-		check(Outputs.Num() == 1);
+		check(Inputs.Num() >= MinAllowedInputTensors);
+		check(Inputs.Num() <= MaxAllowedInputTensors);
+		check(Outputs.Num() == NumAllowedOutputTensors);
 
 		InputIndexDataType = ENNETensorDataType::None;
 
diff --git a/Engine/Plugins/Experimental/NNERuntimeRDG/Source/NNERuntimeRDG/Private/Dml/NNEDmlOperatorSpaceToDepth.cpp b/Engine/Plugins/Experimental/NNERuntimeRDG/Source/NNERuntimeRDG/Private/Dml/NNEDmlOperatorSpaceToDepth.cpp
--- a/Engine/Plugins/Experimental/NNERuntimeRDG/Source/NNERuntimeRDG/Private/Dml/NNEDmlOperatorSpaceToDepth.cpp
+++ b/Engine/Plugins/Experimental/NNERuntimeRDG/Source/NNERuntimeRDG/Private/Dml/NNEDmlOperatorSpaceToDepth.cpp
@@ -18,6 +18,9 @@ class FOperatorDmlSpaceToDepth : public FOperatorDml
 
 	int32 BlockSize;
 
+	static constexpr uint32 NumAllowedInputTensors = 1, NumAllowedOutputTensors = 1;
+	static constexpr int32 	MinTensorRank = 4, MaxTensorRank = 4;
+
 public:
 
 	static FOperatorDml* Create()
@@ -27,13 +30,33 @@ public:
 
 	static bool Validate(const NNE::FAttributeMap& AttributeMap, TConstArrayView<ENNETensorDataType> InputTypes, TConstArrayView<NNE::FSymbolicTensorShape> InputShapes)
 	{
+		const FString OpName = TEXT("SpaceToDepth");
+
+		if(InputShapes.Num() != NumAllowedInputTensors)
+		{
+			UE_LOG(LogNNE, Warning, TEXT("DML %s: Invalid number of input tensors. %d provided, it should be %d."), *OpName, InputShapes.Num(), NumAllowedInputTensors);
+			return false;
+		}
+		
+		if (!CheckGenericTensor(OpName, InputTypes[0], InputShapes[0], 
+			{ 	ENNETensorDataType::Double, ENNETensorDataType::Float, ENNETensorDataType::Half, 
+				ENNETensorDataType::Int64, ENNETensorDataType::Int32, ENNETensorDataType::Int16,
+				ENNETensorDataType::Int8, ENNETensorDataType::UInt64, ENNETensorDataType::UInt32, 
+				ENNETensorDataType::UInt16, ENNETensorDataType::UInt8
+			},
+			MinTensorRank, MaxTensorRank
+		  	))
+		{
+			return false;
+		}
+
 		return true;
 	}
 
 	virtual bool Initialize(TConstArrayView<NNE::FTensorDesc> Inputs, TConstArrayView<NNE::FTensorDesc> Outputs, const NNE::FAttributeMap& Attributes) override
 	{
-		check(Inputs.Num() == 1);
-		check(Outputs.Num() == 1);
+		check(Inputs.Num() == NumAllowedInputTensors);
+		check(Outputs.Num() == NumAllowedOutputTensors);
 
 		const FNNEAttributeValue* BlockSizeAttr = Attributes.GetAttributeValue(TEXT("blocksize"));
 		if (BlockSizeAttr)
@@ -74,6 +97,7 @@ public:
 		FTensorDescDml	DmlOutputTensorDesc;
 
 		if (!DmlInputTensorDesc
+				.SetTensorRank(MinTensorRank, MaxTensorRank)
 				.SetFromTensor(InputTensor)
 				.Validate())
 		{
@@ -82,6 +106,7 @@ public:
 		}
 
 		if (!DmlOutputTensorDesc
+				.SetTensorRank(MinTensorRank, MaxTensorRank)
 				.SetFromTensor(OutputTensor)
 				.Validate())
 		{
diff --git a/Engine/Plugins/Experimental/NNERuntimeRDG/Source/NNERuntimeRDG/Private/Dml/NNEDmlOperatorSplit.cpp b/Engine/Plugins/Experimental/NNERuntimeRDG/Source/NNERuntimeRDG/Private/Dml/NNEDmlOperatorSplit.cpp
--- a/Engine/Plugins/Experimental/NNERuntimeRDG/Source/NNERuntimeRDG/Private/Dml/NNEDmlOperatorSplit.cpp
+++ b/Engine/Plugins/Experimental/NNERuntimeRDG/Source/NNERuntimeRDG/Private/Dml/NNEDmlOperatorSplit.cpp
@@ -11,6 +11,8 @@ class FOperatorDmlSplit : public FOperatorDml
 {
 	TArray<int32>	Split;
 	int				Axis;
+	static constexpr uint32 NumAllowedInputTensors = 1;
+	static constexpr int32 	MinTensorRank = 0, MaxTensorRank = GMaxTensorRank;
 
 public:
 
@@ -21,12 +23,32 @@ public:
 
 	static bool Validate(const NNE::FAttributeMap& AttributeMap, TConstArrayView<ENNETensorDataType> InputTypes, TConstArrayView<NNE::FSymbolicTensorShape> InputShapes)
 	{
+		const FString OpName = TEXT("Split");
+
+		if(InputShapes.Num() != NumAllowedInputTensors)
+		{
+			UE_LOG(LogNNE, Warning, TEXT("DML %s: Invalid number of input tensors. %d provided, it should be %d."), *OpName, InputShapes.Num(), NumAllowedInputTensors);
+			return false;
+		}
+		
+		if (!CheckGenericTensor(OpName, InputTypes[0], InputShapes[0], 
+			{ 	ENNETensorDataType::Double, ENNETensorDataType::Float, ENNETensorDataType::Half, 
+				ENNETensorDataType::Int64, ENNETensorDataType::Int32, ENNETensorDataType::Int16,
+				ENNETensorDataType::Int8, ENNETensorDataType::UInt64, ENNETensorDataType::UInt32, 
+				ENNETensorDataType::UInt16, ENNETensorDataType::UInt8
+			},
+			MinTensorRank, MaxTensorRank
+		  	))
+		{
+			return false;
+		}
+
 		return true;
 	}
 
 	virtual bool Initialize(TConstArrayView<NNE::FTensorDesc> Inputs, TConstArrayView<NNE::FTensorDesc> Outputs, const NNE::FAttributeMap& Attributes) override
 	{
-		check(Inputs.Num() == 1);
+		check(Inputs.Num() == NumAllowedInputTensors);
 		
 		Axis = Attributes.GetValueOrDefault<int>(TEXT("axis"), 0);
 		Axis = HandleNegativeAxis(Axis, Inputs[0].GetShape().Rank());
diff --git a/Engine/Plugins/Experimental/NNERuntimeRDG/Source/NNERuntimeRDG/Private/Dml/NNEDmlOperatorSqueeze.cpp b/Engine/Plugins/Experimental/NNERuntimeRDG/Source/NNERuntimeRDG/Private/Dml/NNEDmlOperatorSqueeze.cpp
--- a/Engine/Plugins/Experimental/NNERuntimeRDG/Source/NNERuntimeRDG/Private/Dml/NNEDmlOperatorSqueeze.cpp
+++ b/Engine/Plugins/Experimental/NNERuntimeRDG/Source/NNERuntimeRDG/Private/Dml/NNEDmlOperatorSqueeze.cpp
@@ -11,6 +11,8 @@ namespace UE::NNERuntimeRDG::Private::Dml
 class FOperatorDmlSqueeze : public FOperatorDml
 {
 	mutable TArray<int32>	Axes;
+	static constexpr uint32 NumAllowedInputTensors = 1, NumAllowedOutputTensors = 1;
+	static constexpr int32 	MinTensorRank = 0, MaxTensorRank = GMaxTensorRank;
 
 public:
 
@@ -21,13 +23,33 @@ public:
 
 	static bool Validate(const NNE::FAttributeMap& AttributeMap, TConstArrayView<ENNETensorDataType> InputTypes, TConstArrayView<NNE::FSymbolicTensorShape> InputShapes)
 	{
+		const FString OpName = TEXT("Squeeze");
+
+		if(InputShapes.Num() != NumAllowedInputTensors)
+		{
+			UE_LOG(LogNNE, Warning, TEXT("DML %s: Invalid number of input tensors. %d provided, it should be %d."), *OpName, InputShapes.Num(), NumAllowedInputTensors);
+			return false;
+		}
+		
+		if (!CheckGenericTensor(OpName, InputTypes[0], InputShapes[0], 
+			{ 	ENNETensorDataType::Double, ENNETensorDataType::Float, ENNETensorDataType::Half, 
+				ENNETensorDataType::Int64, ENNETensorDataType::Int32, ENNETensorDataType::Int16,
+				ENNETensorDataType::Int8, ENNETensorDataType::UInt64, ENNETensorDataType::UInt32, 
+				ENNETensorDataType::UInt16, ENNETensorDataType::UInt8
+			},
+			MinTensorRank, MaxTensorRank
+		  	))
+		{
+			return false;
+		}
+
 		return true;
 	}
 
 	virtual bool Initialize(TConstArrayView<NNE::FTensorDesc> Inputs, TConstArrayView<NNE::FTensorDesc> Outputs, const NNE::FAttributeMap& Attributes) override
 	{
-		check(Inputs.Num() == 1);
-		check(Outputs.Num() == 1);
+		check(Inputs.Num() == NumAllowedInputTensors);
+		check(Outputs.Num() == NumAllowedOutputTensors);
 
 		const int32					InputShapeRank = Inputs[0].GetShape().Rank();
 		const FNNEAttributeValue*	AxesAttr = Attributes.GetAttributeValue(TEXT("axes"));
@@ -95,7 +117,7 @@ public:
 				.SetFromTensor(*OutputTensors[0])
 				.Validate())
 		{
-			UE_LOG(LogNNE, Warning, TEXT("Failed to initialize Unsqueeze's output tensor for DML inference"));
+			UE_LOG(LogNNE, Warning, TEXT("Failed to initialize Squeeze's output tensor for DML inference"));
 			return false;
 		}
 
diff --git a/Engine/Plugins/Experimental/NNERuntimeRDG/Source/NNERuntimeRDG/Private/Dml/NNEDmlOperatorTranspose.cpp b/Engine/Plugins/Experimental/NNERuntimeRDG/Source/NNERuntimeRDG/Private/Dml/NNEDmlOperatorTranspose.cpp
--- a/Engine/Plugins/Experimental/NNERuntimeRDG/Source/NNERuntimeRDG/Private/Dml/NNEDmlOperatorTranspose.cpp
+++ b/Engine/Plugins/Experimental/NNERuntimeRDG/Source/NNERuntimeRDG/Private/Dml/NNEDmlOperatorTranspose.cpp
@@ -3,6 +3,7 @@
 #ifdef NNE_USE_DIRECTML
 
 #include "NNEDmlOperator.h"
+#include "NNEDmlOperatorUtils.h"
 #include "Math/Range.h"
 #include "Algo/Transform.h"
 #include "Algo/Reverse.h"
@@ -27,6 +28,9 @@ class FOperatorDmlTranspose : public FOperatorDml
 
 	TArray<int32> Perm;
 
+	static constexpr uint32 NumAllowedInputTensors = 1, NumAllowedOutputTensors = 1;
+	static constexpr int32 	MinTensorRank = 0, MaxTensorRank = GMaxTensorRank;
+
 public:
 
 	static FOperatorDml* Create()
@@ -36,13 +40,33 @@ public:
 
 	static bool Validate(const NNE::FAttributeMap& AttributeMap, TConstArrayView<ENNETensorDataType> InputTypes, TConstArrayView<NNE::FSymbolicTensorShape> InputShapes)
 	{
+		const FString OpName = TEXT("Transpose");
+
+		if(InputShapes.Num() != NumAllowedInputTensors)
+		{
+			UE_LOG(LogNNE, Warning, TEXT("DML %s: Invalid number of input tensors. %d provided, it should be %d."), *OpName, InputShapes.Num(), NumAllowedInputTensors);
+			return false;
+		}
+		
+		if (!CheckGenericTensor(OpName, InputTypes[0], InputShapes[0], 
+			{ 	ENNETensorDataType::Double, ENNETensorDataType::Float, ENNETensorDataType::Half, 
+				ENNETensorDataType::Int64, ENNETensorDataType::Int32, ENNETensorDataType::Int16,
+				ENNETensorDataType::Int8, ENNETensorDataType::UInt64, ENNETensorDataType::UInt32, 
+				ENNETensorDataType::UInt16, ENNETensorDataType::UInt8
+			},
+			MinTensorRank, MaxTensorRank
+		  	))
+		{
+			return false;
+		}
+
 		return true;
 	}
 
 	virtual bool Initialize(TConstArrayView<NNE::FTensorDesc> Inputs, TConstArrayView<NNE::FTensorDesc> Outputs, const NNE::FAttributeMap& Attributes) override
 	{
-		check(Inputs.Num() == 1);
-		check(Outputs.Num() == 1);
+		check(Inputs.Num() == NumAllowedInputTensors);
+		check(Outputs.Num() == NumAllowedOutputTensors);
 
 		const int32 NumDims = Inputs[0].GetShape().Rank();
 		check(NumDims > 0);
diff --git a/Engine/Plugins/Experimental/NNERuntimeRDG/Source/NNERuntimeRDG/Private/Dml/NNEDmlOperatorUnsqueeze.cpp b/Engine/Plugins/Experimental/NNERuntimeRDG/Source/NNERuntimeRDG/Private/Dml/NNEDmlOperatorUnsqueeze.cpp
--- a/Engine/Plugins/Experimental/NNERuntimeRDG/Source/NNERuntimeRDG/Private/Dml/NNEDmlOperatorUnsqueeze.cpp
+++ b/Engine/Plugins/Experimental/NNERuntimeRDG/Source/NNERuntimeRDG/Private/Dml/NNEDmlOperatorUnsqueeze.cpp
@@ -11,6 +11,8 @@ namespace UE::NNERuntimeRDG::Private::Dml
 class FOperatorDmlUnsqueeze : public FOperatorDml
 {
 	TArray<int32> Axes;
+	static constexpr uint32 NumAllowedInputTensors = 1, NumAllowedOutputTensors = 1;
+	static constexpr int32 	MinTensorRank = 0, MaxTensorRank = GMaxTensorRank;
 	
 public:
 
@@ -21,13 +23,33 @@ public:
 
 	static bool Validate(const NNE::FAttributeMap& AttributeMap, TConstArrayView<ENNETensorDataType> InputTypes, TConstArrayView<NNE::FSymbolicTensorShape> InputShapes)
 	{
+		const FString OpName = TEXT("Unsqueeze");
+
+		if(InputShapes.Num() != NumAllowedInputTensors)
+		{
+			UE_LOG(LogNNE, Warning, TEXT("DML %s: Invalid number of input tensors. %d provided, it should be %d."), *OpName, InputShapes.Num(), NumAllowedInputTensors);
+			return false;
+		}
+		
+		if (!CheckGenericTensor(OpName, InputTypes[0], InputShapes[0], 
+			{ 	ENNETensorDataType::Double, ENNETensorDataType::Float, ENNETensorDataType::Half, 
+				ENNETensorDataType::Int64, ENNETensorDataType::Int32, ENNETensorDataType::Int16,
+				ENNETensorDataType::Int8, ENNETensorDataType::UInt64, ENNETensorDataType::UInt32, 
+				ENNETensorDataType::UInt16, ENNETensorDataType::UInt8
+			},
+			MinTensorRank, MaxTensorRank
+		  	))
+		{
+			return false;
+		}
+
 		return true;
 	}
 
 	virtual bool Initialize(TConstArrayView<NNE::FTensorDesc> Inputs, TConstArrayView<NNE::FTensorDesc> Outputs, const NNE::FAttributeMap& Attributes) override
 	{
-		check(Inputs.Num() == 1);
-		check(Outputs.Num() == 1);
+		check(Inputs.Num() == NumAllowedInputTensors);
+		check(Outputs.Num() == NumAllowedOutputTensors);
 
 		const FNNEAttributeValue* AxesAttr = Attributes.GetAttributeValue(TEXT("axes"));
 		if (!AxesAttr)
diff --git a/Engine/Plugins/Experimental/NNERuntimeRDG/Source/NNERuntimeRDG/Private/Dml/NNEDmlOperatorUtils.h b/Engine/Plugins/Experimental/NNERuntimeRDG/Source/NNERuntimeRDG/Private/Dml/NNEDmlOperatorUtils.h
--- a/Engine/Plugins/Experimental/NNERuntimeRDG/Source/NNERuntimeRDG/Private/Dml/NNEDmlOperatorUtils.h
+++ b/Engine/Plugins/Experimental/NNERuntimeRDG/Source/NNERuntimeRDG/Private/Dml/NNEDmlOperatorUtils.h
@@ -83,33 +83,34 @@ static bool IsEqualOrBroadcastable(TConstArrayView<T> ShapeA, TConstArrayView<T>
 	return true;
 }
 
-static bool CheckGenericTensor(ENNETensorDataType DataType, const NNE::FSymbolicTensorShape& TensorShape)
+static bool CheckGenericTensor(const FString& OpName, ENNETensorDataType DataType, const NNE::FSymbolicTensorShape& TensorShape, 
+							   const TSet<ENNETensorDataType>& AllowedDataTypes,
+							   const int32 MinTensorRank = 0, const int32 MaxTensorRank = GMaxTensorRank)
 {
-	const int32 MinTensorRank(0), MaxTensorRank(DML_TENSOR_DIMENSION_COUNT_MAX1);
-
+	if(!AllowedDataTypes.Contains(DataType))
+	{
+		UE_LOG(LogNNE, Warning, TEXT("DML %s: Invalid data type: %s"), *OpName, *UEnum::GetDisplayValueAsText(DataType).ToString());
+		return false;
+	}
 	if (TensorShape.Rank() < MinTensorRank || TensorShape.Rank() > MaxTensorRank)
 	{
-		UE_LOG(LogNNE, Warning, TEXT("Invalid DML tensor rank: %d [%d,%d]"), TensorShape.Rank(), MinTensorRank, MaxTensorRank);
+		UE_LOG(LogNNE, Warning, TEXT("DML %s: Invalid tensor rank: %d [%d,%d]"), *OpName, TensorShape.Rank(), MinTensorRank, MaxTensorRank);
 		return false;
 	}
 
 	return true;
 }
 
-static bool CheckElementwiseTensor(ENNETensorDataType DataType, const NNE::FSymbolicTensorShape& TensorShape)
+static bool CheckGenericTensor1D(const FString& OpName, ENNETensorDataType DataType, const NNE::FSymbolicTensorShape& TensorShape, 
+							   const TSet<ENNETensorDataType>& AllowedDataTypes)
 {
-	if (DataType != ENNETensorDataType::Float)
-	{
-		UE_LOG(LogNNE, Warning, TEXT("Invalid DML tensor data type"));
-		return false;
-	}
-
-	if(!CheckGenericTensor(DataType, TensorShape))
-	{
-		return false;
-	}
+	return CheckGenericTensor(OpName, DataType, TensorShape, AllowedDataTypes, 0, 1); // Scalar is "cast" to 1D due to lack of support in Dml.
+}
 
-	return true;
+static bool CheckElementwiseTensor(const FString& OpName, ENNETensorDataType DataType, const NNE::FSymbolicTensorShape& TensorShape,
+								   TSet<ENNETensorDataType> AllowedDataTypes = {ENNETensorDataType::Float, ENNETensorDataType::Half})
+{
+	return CheckGenericTensor(OpName, DataType, TensorShape, AllowedDataTypes);
 }
 
 static Util::FSmallUIntArray KernelPadding(
