[NNE] Don't create RDGDml model instance when unsupported RHI is specified on the command line

Problem
If command line has "-vulkan" or "-dx11" on Win64 platform, CreateInstanceRDG() is called which should not happen, because DirectML can only be used on D3D12 RHI

Solution
Prevent creating a RDGDml model instance in this case

Tests
NNEIncubator Editor Win64

#rnx

diff --git a/Engine/Plugins/Experimental/NNERuntimeRDG/Source/NNERuntimeRDG/Private/NNERuntimeRDGDml.cpp b/Engine/Plugins/Experimental/NNERuntimeRDG/Source/NNERuntimeRDG/Private/NNERuntimeRDGDml.cpp
--- a/Engine/Plugins/Experimental/NNERuntimeRDG/Source/NNERuntimeRDG/Private/NNERuntimeRDGDml.cpp
+++ b/Engine/Plugins/Experimental/NNERuntimeRDG/Source/NNERuntimeRDG/Private/NNERuntimeRDGDml.cpp
@@ -30,7 +30,7 @@ namespace UE::NNERuntimeRDG::Private::Dml
 bool FRuntimeDmlStartup()
 {
 	bool bIsD3D12RHI = GDynamicRHI && GDynamicRHI->GetInterfaceType() == ERHIInterfaceType::D3D12;		
-	bool bLoadDirectML = true;
+	bool bLoadDirectML = bIsD3D12RHI;
 
 	if (IsRunningCommandlet() && !IsAllowCommandletRendering())
 	{
@@ -155,6 +155,11 @@ FString UNNERuntimeRDGDmlImpl::GetModelDataIdentifier(FString FileType, TConstAr
 bool UNNERuntimeRDGDmlImpl::CanCreateModelRDG(TObjectPtr<UNNEModelData> ModelData) const
 {
 #ifdef NNE_USE_DIRECTML
+	if (bRegisterOnlyOperators)
+	{
+		return false;
+	}
+
 	TSharedPtr<UE::NNE::FSharedModelData> SharedData = ModelData->GetModelData(GetRuntimeName());
 
 	if (!SharedData.IsValid())
@@ -215,13 +220,15 @@ UNNERuntimeRDGDmlImpl::~UNNERuntimeRDGDmlImpl()
 #endif
 }
 
-bool UNNERuntimeRDGDmlImpl::Init(bool bRegisterOnlyOperators)
+bool UNNERuntimeRDGDmlImpl::Init(bool bInRegisterOnlyOperators)
 {
+	bRegisterOnlyOperators = bInRegisterOnlyOperators;
+
 #ifdef NNE_USE_DIRECTML
 	
 	if (bRegisterOnlyOperators)
 	{
-		UE_LOG(LogNNE, Display, TEXT("Registering only operators"));
+		UE_LOG(LogNNE, Display, TEXT("RDGDml:Registering only operators"));
 		return true;
 	}
 
diff --git a/Engine/Plugins/Experimental/NNERuntimeRDG/Source/NNERuntimeRDG/Private/NNERuntimeRDGDml.h b/Engine/Plugins/Experimental/NNERuntimeRDG/Source/NNERuntimeRDG/Private/NNERuntimeRDGDml.h
--- a/Engine/Plugins/Experimental/NNERuntimeRDG/Source/NNERuntimeRDG/Private/NNERuntimeRDGDml.h
+++ b/Engine/Plugins/Experimental/NNERuntimeRDG/Source/NNERuntimeRDG/Private/NNERuntimeRDGDml.h
@@ -39,7 +39,10 @@ public:
 
 	virtual bool CanCreateModelRDG(TObjectPtr<UNNEModelData> ModelData) const override;
 	virtual TSharedPtr<UE::NNE::IModelRDG> CreateModelRDG(TObjectPtr<UNNEModelData> ModelData) override;
+
 private:
+	
 	UE::NNERuntimeRDG::Private::Dml::FDmlDeviceContext* Ctx;
+	bool												bRegisterOnlyOperators { false };
 };
 
