[NNE] Add tracing events for RDGDml

Problem:
There are GPU stats for RDGDml FModelInstance, but the inference GPU time is not displayed.
Because of that we're not aware if there are any performance regressions.

Solution:
Use FRealtimeGPUProfiler::FetchPerfByDescription() in the NNEIncubator and display it on the using AddOnScreenDebugMessage() and into log.
To ease debugging, add trace log events to visualize each method of FModelInstance (on which thread game, render, RHI thread) in the UnrealInsights.

Tests
NNEIncubator Editor Win64

#rnx

diff --git a/Engine/Plugins/Experimental/NNERuntimeRDG/Source/NNERuntimeRDG/Internal/NNERenderGraphUtils.h b/Engine/Plugins/Experimental/NNERuntimeRDG/Source/NNERuntimeRDG/Internal/NNERenderGraphUtils.h
new file mode 100644
--- /dev/null
+++ b/Engine/Plugins/Experimental/NNERuntimeRDG/Source/NNERuntimeRDG/Internal/NNERenderGraphUtils.h
@@ -0,0 +1,206 @@
+// Copyright Epic Games, Inc. All Rights Reserved.
+
+#pragma once
+
+#include "NNERuntimeCPU.h"
+#include "NNERuntimeRDG.h"
+#include "RenderGraph.h"
+
+#if UE_TRACE_ENABLED
+	#define NNE_TRACE_EVENT_SCOPED(Name) SCOPED_NAMED_EVENT_TEXT(#Name, FColor::Green)
+#else
+	#define NNE_TRACE_EVENT_SCOPED(Name)
+#endif
+
+namespace UE::NNE::Internal
+{
+
+/*
+* Utility class to manage GPU -> CPU readbacks
+*/
+class FReadbackManager
+{
+	struct FReadback
+	{
+		FStagingBufferRHIRef	StagingBuffer;
+		FGPUFenceRHIRef			Fence;
+		FRDGBufferRef			Buffer;
+		void*					DstData;
+		uint32					NumBytes;
+	};
+
+	BEGIN_SHADER_PARAMETER_STRUCT(FReadbackPassParameters, )
+		RDG_BUFFER_ACCESS_ARRAY(ReadbackBuffers)
+	END_SHADER_PARAMETER_STRUCT()
+
+public:
+
+	FReadbackManager()
+	{
+		check(IsInGameThread());
+		Signal = FGenericPlatformProcess::GetSynchEventFromPool(true);
+	}
+
+	~FReadbackManager()
+	{
+		check(IsInGameThread());
+		FGenericPlatformProcess::ReturnSynchEventToPool(Signal);
+		Signal = nullptr;
+	}
+
+	// Note: call on the game thread
+	bool Init(int32 InNumReadbacks)
+	{
+		NNE_TRACE_EVENT_SCOPED(NNE_FTensorReadback_Init);
+
+		check(IsInGameThread());
+		MaxNumReadbacks = InNumReadbacks;
+		Signal->Reset();
+
+		return true;
+	}
+
+	void BeginReadbacks_RenderThread(FRDGBuilder& RDGBuilder)
+	{
+		NNE_TRACE_EVENT_SCOPED(NNE_FTensorReadback_BeginReadbacks);
+
+		check(IsInRenderingThread());
+		check(Readbacks.IsEmpty());
+		check(MaxNumReadbacks > 0);
+
+		NumAddedReadbacks = 0;
+		NumProcessedReadbacks = 0;
+		Readbacks.SetNum(MaxNumReadbacks);
+		
+		ReadbackParams = RDGBuilder.AllocParameters<FReadbackPassParameters>();
+
+		for (FReadback& Readback : Readbacks)
+		{
+			Readback.StagingBuffer = RHICreateStagingBuffer();
+			Readback.StagingBuffer->DisableLifetimeExtension();
+
+			Readback.Fence = RHICreateGPUFence(TEXT("FReadbackManager_ReadbackFence"));
+			Readback.Fence->DisableLifetimeExtension();
+			Readback.Fence->Clear();
+
+			Readback.Buffer = nullptr;
+			Readback.DstData = nullptr;
+			Readback.NumBytes = 0;
+		}
+	}
+
+	void AddReadbacks_RenderThread(TConstArrayView<FTensorBindingRDG> InBindingsRDG, TConstArrayView<FTensorBindingCPU> InBindingsCPU)
+	{
+		NNE_TRACE_EVENT_SCOPED(NNE_FTensorReadback_AddReadbacks_RT);
+
+		check(IsInRenderingThread());
+		check(InBindingsRDG.Num() == InBindingsCPU.Num());
+		check(InBindingsRDG.Num() <= MaxNumReadbacks);
+
+		if (InBindingsRDG.Num() != InBindingsCPU.Num())
+		{
+			UE_LOG(LogNNE, Error, TEXT("FReadbackManager:Number of bindings need to be same"));
+			return;
+		}
+
+		if (InBindingsRDG.Num() > MaxNumReadbacks)
+		{
+			UE_LOG(LogNNE, Error, TEXT("FReadbackManager:Number of bindings is larger than the maximum number of readbacks provided by the Init()"));
+			return;
+		}
+
+		for (int32 Idx = 0; Idx < InBindingsRDG.Num(); ++Idx)
+		{
+			ReadbackParams->ReadbackBuffers.Emplace(InBindingsRDG[Idx].Buffer, ERHIAccess::CopySrc);
+
+			FReadback& Readback = Readbacks[NumAddedReadbacks];
+			Readback.Buffer = InBindingsRDG[Idx].Buffer;
+			Readback.DstData = InBindingsCPU[Idx].Data;
+			Readback.NumBytes = InBindingsCPU[Idx].SizeInBytes;
+
+			NumAddedReadbacks++;
+		}
+	}
+
+	void EndReadbacks_RenderThread(FRDGBuilder& RDGBuilder)
+	{
+		check(IsInRenderingThread());
+
+		RDGBuilder.AddPass(
+			RDG_EVENT_NAME("FReadbackManager_ProcessReadbacks"),
+			ReadbackParams,
+			ERDGPassFlags::Readback | ERDGPassFlags::NeverCull,
+			[this](FRHICommandListImmediate& RHICmdList)
+			{
+				NNE_TRACE_EVENT_SCOPED(NNE_FTensorReadback_EndReadbacks_RT);
+				
+				for (FReadback& Readback : Readbacks)
+				{
+					Readback.Buffer->MarkResourceAsUsed();
+
+					RHICmdList.CopyToStagingBuffer(Readback.Buffer->GetRHI(), Readback.StagingBuffer, 0, Readback.NumBytes);					
+					RHICmdList.WriteGPUFence(Readback.Fence);
+				}
+
+				RHICmdList.BlockUntilGPUIdle();
+
+				for (FReadback& Readback : Readbacks)
+				{
+					const void* SrcData = RHICmdList.LockStagingBuffer(Readback.StagingBuffer, Readback.Fence.GetReference(), 0, Readback.NumBytes);
+					check(SrcData);
+
+					if (SrcData)
+					{
+						FMemory::Memcpy(Readback.DstData, SrcData, Readback.NumBytes);
+						RHICmdList.UnlockStagingBuffer(Readback.StagingBuffer);
+
+						Readback.StagingBuffer = nullptr;
+						Readback.Fence = nullptr;
+					}
+
+					++NumProcessedReadbacks;
+				}
+				
+				// Clean-up resources (staging buffers, fences)
+				RHICmdList.ImmediateFlush(EImmediateFlushType::FlushRHIThreadFlushResources);
+
+				Signal->Trigger();
+			}
+		);
+	}
+
+	// Note: call on the game thread
+	void Wait()
+	{
+		check(IsInGameThread());
+
+		NNE_TRACE_EVENT_SCOPED(NNE_FTensorReadback_Wait);
+
+		check(MaxNumReadbacks > 0);
+		Signal->Wait();
+
+		while (NumProcessedReadbacks < MaxNumReadbacks)
+		{
+			FPlatformProcess::Sleep(0.f);
+		}
+
+		NumAddedReadbacks = 0;
+		NumProcessedReadbacks = 0;
+		MaxNumReadbacks = 0;
+		Readbacks.Empty();
+	}
+
+
+private:
+
+	using ReadbackArray = TArray<FReadback, TInlineAllocator<16>>;
+	
+	ReadbackArray				Readbacks;
+	FReadbackPassParameters*	ReadbackParams = nullptr;
+	std::atomic<int32>			MaxNumReadbacks;
+	std::atomic<int32>			NumAddedReadbacks;
+	std::atomic<int32>			NumProcessedReadbacks;
+	FEvent*						Signal = nullptr;
+};
+
+} // namespace UE::NNE
diff --git a/Engine/Plugins/Experimental/NNERuntimeRDG/Source/NNERuntimeRDG/NNERuntimeRDG.Build.cs b/Engine/Plugins/Experimental/NNERuntimeRDG/Source/NNERuntimeRDG/NNERuntimeRDG.Build.cs
--- a/Engine/Plugins/Experimental/NNERuntimeRDG/Source/NNERuntimeRDG/NNERuntimeRDG.Build.cs
+++ b/Engine/Plugins/Experimental/NNERuntimeRDG/Source/NNERuntimeRDG/NNERuntimeRDG.Build.cs
@@ -27,8 +27,9 @@ public class NNERuntimeRDG : ModuleRules
             "NNE",
 			"NNEHlslShaders",
             "RHI",
-			"Projects"
-        });
+			"Projects",
+			"TraceLog"
+		});
 
         if (Target.Platform == UnrealTargetPlatform.Win64)
 		{
diff --git a/Engine/Plugins/Experimental/NNERuntimeRDG/Source/NNERuntimeRDG/Private/Dml/NNEDmlModel.cpp b/Engine/Plugins/Experimental/NNERuntimeRDG/Source/NNERuntimeRDG/Private/Dml/NNEDmlModel.cpp
--- a/Engine/Plugins/Experimental/NNERuntimeRDG/Source/NNERuntimeRDG/Private/Dml/NNEDmlModel.cpp
+++ b/Engine/Plugins/Experimental/NNERuntimeRDG/Source/NNERuntimeRDG/Private/Dml/NNEDmlModel.cpp
@@ -7,6 +7,7 @@
 
 #include "RenderGraphBuilder.h"
 #include "RenderGraphUtils.h"
+#include "NNERenderGraphUtils.h"
 
 #include "ID3D12DynamicRHI.h"
 #endif
@@ -22,18 +23,14 @@ static FAutoConsoleVariableRef CVarNNEDmlMetaCommands(
 namespace UE::NNERuntimeRDG::Private::Dml
 {
 
-#if STATS
-DECLARE_STATS_GROUP(TEXT("NNEDmlModel"), STATGROUP_NNEDmlModel, STATCAT_Advanced);
-DECLARE_MEMORY_STAT(TEXT("MemSizeWeights GPU"), STAT_MemSizeWeights, STATGROUP_NNEDmlModel);
-DECLARE_MEMORY_STAT(TEXT("MemSizePersistent GPU"), STAT_MemSizePersist, STATGROUP_NNEDmlModel);
-DECLARE_MEMORY_STAT(TEXT("MemSizeTemp GPU"), STAT_MemSizeTemp, STATGROUP_NNEDmlModel);
-DECLARE_DWORD_COUNTER_STAT(TEXT("InferenceCount"), STAT_InferenceCount, STATGROUP_NNEDmlModel);
-#endif
+DECLARE_STATS_GROUP(TEXT("NNEDmlModelInstance"), STATGROUP_NNEDmlModelInstance, STATCAT_Advanced);
+DECLARE_MEMORY_STAT(TEXT("MemSizeWeights GPU"), STAT_MemSizeWeights, STATGROUP_NNEDmlModelInstance);
+DECLARE_MEMORY_STAT(TEXT("MemSizePersistent GPU"), STAT_MemSizePersist, STATGROUP_NNEDmlModelInstance);
+DECLARE_MEMORY_STAT(TEXT("MemSizeTemp GPU"), STAT_MemSizeTemp, STATGROUP_NNEDmlModelInstance);
+DECLARE_DWORD_COUNTER_STAT(TEXT("InferenceCount"), STAT_InferenceCount, STATGROUP_NNEDmlModelInstance);
 
-#if HAS_GPU_STATS
-DECLARE_GPU_STAT_NAMED(GPU_STAT_DispatchTime, TEXT("NNEDmlModelDispatchTime"));
-DECLARE_GPU_STAT_NAMED(GPU_STAT_DispatchD3DTime, TEXT("NNEDmlModelDispatchD3DTime"));
-#endif
+DECLARE_GPU_STAT_NAMED(NNE_DmlModelInstance_Dispatch_GPU, TEXT("NNE_DmlModelInstance_Dispatch"));
+DECLARE_GPU_STAT_NAMED(NNE_DmlModelInstance_DispatchD3D_GPU, TEXT("NNE_DmlModelInstance_DispatchD3D"));
 
 // Empty tensor is marked by -1
 static constexpr int32 GEmptyTensorIdx = -1;
@@ -824,18 +821,20 @@ FModelInstance::FModelInstance()
 
 FModelInstance::~FModelInstance()
 {
-#if STATS
+	NNE_TRACE_EVENT_SCOPED(NNE_DmlModel_<dtor>);
+
 	DEC_MEMORY_STAT_BY(STAT_MemSizeWeights, MemSizeWeights);
 	DEC_MEMORY_STAT_BY(STAT_MemSizeTemp, MemSizeTemp);
 	DEC_MEMORY_STAT_BY(STAT_MemSizePersist, MemSizePersist);
-#endif
 	
 	FEvent* Signal = FGenericPlatformProcess::GetSynchEventFromPool(false);
 
-	ENQUEUE_RENDER_COMMAND(FDmlModelInstance_Release)
+	ENQUEUE_RENDER_COMMAND(NNE_DmlModel_ReleaseResources)
 	(
 		[this, &Signal](FRHICommandListImmediate& RHICmdList)
 		{
+			NNE_TRACE_EVENT_SCOPED(NNE_DmlModel_ReleaseResources_RT);
+
 			if (DispatchFence.IsValid())
 			{
 				while (!DispatchFence->Poll())
@@ -1003,6 +1002,8 @@ bool FModelInstance::InitCompiledOp()
 	static constexpr EBufferUsageFlags	TempBuffFlags = BUF_Static | BUF_ShaderResource | BUF_UnorderedAccess;
 	static constexpr ERHIAccess			TempBuffAccess = ERHIAccess::UAVMask;
 
+	NNE_TRACE_EVENT_SCOPED(NNE_DmlModelInstance_InitCompiledOp);
+
 	HRESULT					Res;
 	IDMLDevice*				Device = DevCtx->Device;
 	IDMLCompiledOperator*	CompiledOps[] = { CompiledOp };
@@ -1045,7 +1046,7 @@ bool FModelInstance::InitCompiledOp()
 
 	FEvent* Signal = FGenericPlatformProcess::GetSynchEventFromPool(false);
 	
-	ENQUEUE_RENDER_COMMAND(FDmlModelInstance_SetTensorData)
+	ENQUEUE_RENDER_COMMAND(NNE_DmlModelInstance_InitCompiledOp)
 	(
 		[
 			this, 
@@ -1054,7 +1055,9 @@ bool FModelInstance::InitCompiledOp()
 		]
 		(FRHICommandListImmediate& RHICmdList)
 		{
-			DispatchFence = DynamicRHI->RHICreateGPUFence(TEXT("FDmlModelInstanceDispatchFence"));
+			NNE_TRACE_EVENT_SCOPED(NNE_DmlModelInstance_InitCompiledOp_RT);
+
+			DispatchFence = DynamicRHI->RHICreateGPUFence(TEXT("NNE_DmlModelInstance_DispatchFence"));
 
 			FRHIBufferInputArray	Inputs;
 
@@ -1072,9 +1075,9 @@ bool FModelInstance::InitCompiledOp()
 
 			if (MemSizeWeights)
 			{
-				UploadFence = RHICreateGPUFence(TEXT("FDmlModelInstance_UploadFence"));
+				UploadFence = RHICreateGPUFence(TEXT("NNE_DmlModelInstance_UploadFence"));
 
-				UploadBuff = CreateRHIBuffer(RHICmdList, MemSizeWeights, BUF_ShaderResource | BUF_Dynamic | BUF_FastVRAM, ERHIAccess::CopySrc, TEXT("FDmlModelInstance_UploadBuffer"));
+				UploadBuff = CreateRHIBuffer(RHICmdList, MemSizeWeights, BUF_ShaderResource | BUF_Dynamic | BUF_FastVRAM, ERHIAccess::CopySrc, TEXT("NNE_DmlModelInstance_UploadBuffer"));
 				check(UploadBuff);
 				UploadBuff->DisableLifetimeExtension();
 
@@ -1118,14 +1121,14 @@ bool FModelInstance::InitCompiledOp()
 
 			if (MemSizePersist)
 			{
-				PersistBuff = CreateRHIBuffer(RHICmdList, MemSizePersist, PersistBuffFlags, PersistBuffAccess, TEXT("FDmlModelInstance_PeristBuff"));
+				PersistBuff = CreateRHIBuffer(RHICmdList, MemSizePersist, PersistBuffFlags, PersistBuffAccess, TEXT("NNE_DmlModelInstance_PeristBuff"));
 				check(PersistBuff.IsValid());
 				INC_MEMORY_STAT_BY(STAT_MemSizePersist, MemSizePersist);
 			}
 
 			if (MemSizeTemp)
 			{
-				TempBuff = CreateRHIBuffer(RHICmdList, MemSizeTemp, TempBuffFlags, TempBuffAccess, TEXT("FDmlModelInstance_TempBuff"));
+				TempBuff = CreateRHIBuffer(RHICmdList, MemSizeTemp, TempBuffFlags, TempBuffAccess, TEXT("NNE_DmlModelInstance_TempBuff"));
 				check(TempBuff.IsValid());
 				INC_MEMORY_STAT_BY(STAT_MemSizeTemp, MemSizeTemp);
 			}
@@ -1134,7 +1137,7 @@ bool FModelInstance::InitCompiledOp()
 
 			if (InitTempMemSize)
 			{
-				InitTempBuff = CreateRHIBuffer(RHICmdList, InitTempMemSize, TempBuffFlags, TempBuffAccess, TEXT("FDmlModelInstance_InitTempBuff"));
+				InitTempBuff = CreateRHIBuffer(RHICmdList, InitTempMemSize, TempBuffFlags, TempBuffAccess, TEXT("NNE_DmlModelInstance_InitTempBuff"));
 				InitTempBuff->DisableLifetimeExtension();
 				check(InitTempBuff.IsValid());
 			}
@@ -1154,10 +1157,12 @@ bool FModelInstance::InitCompiledOp()
 			RHICmdList.EnqueueLambda(
 				[this, Inputs, Barriers, InitTempBuff](FRHICommandListImmediate& RHICmdList)
 				{
+					NNE_TRACE_EVENT_SCOPED(NNE_DmlModelInstance_InitCompiledOpD3D_RHI);
+
 					ID3D12GraphicsCommandList* D3DCmdList = nullptr;
 
 					D3DCmdList = DynamicRHI->RHIGetGraphicsCommandList(DevCtx->DeviceIndex);
-					D3DCmdList->SetName(TEXT("FDmlModelInstance_SetTensorData_CmdList"));
+					D3DCmdList->SetName(TEXT("NNE_DmlModelInstance_InitCompiledOp_CmdList"));
 
 					BindingTable->Bind(OpInit, Inputs, PersistBuff, InitTempBuff);
 			
@@ -1205,9 +1210,8 @@ END_SHADER_PARAMETER_STRUCT()
 
 void FModelInstance::AddDispatchOps_RenderThread(FRDGBuilder& GraphBuilder)
 {
-#if STATS
+	NNE_TRACE_EVENT_SCOPED(NNE_DmlModelInstance_Dispatch);
 	INC_DWORD_STAT(STAT_InferenceCount);
-#endif
 
 	FDmlModelDispatchPassParameters* DispatchParams = GraphBuilder.AllocParameters<FDmlModelDispatchPassParameters>();
 
@@ -1231,16 +1235,18 @@ void FModelInstance::AddDispatchOps_RenderThread(FRDGBuilder& GraphBuilder)
 		DispatchParams->OutputBuffers.Emplace(AllTensorRDGRefs[OutputTensorIndices[Idx]]->GetBuffer(), ERHIAccess::UAVCompute);
 	}
 
-#if HAS_GPU_STATS
-	RDG_GPU_STAT_SCOPE(GraphBuilder, GPU_STAT_DispatchTime);
-#endif
+	RDG_EVENT_SCOPE(GraphBuilder, "NNE_DmlModelInstance_Dispatch_Pass");
+	RDG_GPU_STAT_SCOPE_VERBOSE(GraphBuilder, NNE_DmlModelInstance_Dispatch_GPU, TEXT("NNE_DmlModelInstance_Dispatch_GPU"));
+	RDG_CSV_STAT_EXCLUSIVE_SCOPE(GraphBuilder, NNE_DmlModelInstance_Dispatch);
 
 	GraphBuilder.AddPass(
-		RDG_EVENT_NAME("FDmlModelInstance_Dispatch"),
+		RDG_EVENT_NAME("NNE_DmlModelInstance_Dispatch_Pass"),
 		DispatchParams,
 		ERDGPassFlags::Compute | ERDGPassFlags::NeverCull,
 		[this, DispatchParams, NumWeightTensors](FRHICommandListImmediate& RHICmdList)
 		{
+			NNE_TRACE_EVENT_SCOPED(NNE_DmlModelInstance_Dispatch_Pass);
+
 			FRHIBufferInputArray	RHIInputBuffers;
 			FRHIBufferOutputArray	RHIOutputBuffers;
 
@@ -1262,14 +1268,14 @@ void FModelInstance::AddDispatchOps_RenderThread(FRDGBuilder& GraphBuilder)
 			}
 
 			{
-#if HAS_GPU_STATS
-				SCOPED_GPU_STAT(RHICmdList, GPU_STAT_DispatchD3DTime)
-#endif
+				SCOPED_GPU_STAT(RHICmdList, NNE_DmlModelInstance_DispatchD3D_GPU)
 
 				DispatchFence->Clear();
 				RHICmdList.EnqueueLambda(
 					[this, InputBuffers = MoveTemp(RHIInputBuffers), OutputBuffers = MoveTemp(RHIOutputBuffers)](FRHICommandListImmediate& RHICmdList)
 					{
+						NNE_TRACE_EVENT_SCOPED(NNE_DmlModelInstance_DispatchD3D_RHI);
+
 						TArray<CD3DX12_RESOURCE_BARRIER, TInlineAllocator<MaxNumInputs + MaxNumOutputs>>	PreBarriers;
 						TArray<CD3DX12_RESOURCE_BARRIER, TInlineAllocator<MaxNumOutputs * 2>>				PostBarriers;
 
@@ -1320,7 +1326,7 @@ void FModelInstance::AddDispatchOps_RenderThread(FRDGBuilder& GraphBuilder)
 						ID3D12GraphicsCommandList* D3DCmdList = nullptr;
 
 						D3DCmdList = DynamicRHI->RHIGetGraphicsCommandList(DevCtx->DeviceIndex);
-						D3DCmdList->SetName(TEXT("FDmlModelInstance_Dispatch_CmdList"));
+						D3DCmdList->SetName(TEXT("NNE_DmlModelInstance_Dispatch_CmdList"));
 						D3DCmdList->SetDescriptorHeaps(1, &DescHeap);
 						D3DCmdList->ResourceBarrier(PreBarriers.Num(), PreBarriers.GetData());
 						DevCtx->CmdRec->RecordDispatch(D3DCmdList, CompiledOp, BindingTable->Get());
