NNE NNERuntimeRDGHlsl concat gpu implementation

Problem:
NNERuntimeRDGHlsl was missing concat on GPU.

Solution:
Implement it

Test:
Win64 Editor: all tests
Win64 Standalone: all tests

#rnx

diff --git a/Engine/Plugins/Experimental/NNERuntimeRDG/Source/NNERuntimeRDG/Private/Hlsl/NNERuntimeRDGConcat.cpp b/Engine/Plugins/Experimental/NNERuntimeRDG/Source/NNERuntimeRDG/Private/Hlsl/NNERuntimeRDGConcat.cpp
--- a/Engine/Plugins/Experimental/NNERuntimeRDG/Source/NNERuntimeRDG/Private/Hlsl/NNERuntimeRDGConcat.cpp
+++ b/Engine/Plugins/Experimental/NNERuntimeRDG/Source/NNERuntimeRDG/Private/Hlsl/NNERuntimeRDGConcat.cpp
@@ -4,9 +4,39 @@
 #include "NNETensor.h"
 #include "NNETypes.h"
 #include "NNERuntimeRDGHelperConcat.h"
+#include "RenderGraphUtils.h"
 
 namespace UE::NNERuntimeRDG::Private::Hlsl
 {
+	namespace ConcatHelper
+	{
+		uint64 GetNumElemLeftOfAxisForShape(const NNE::FTensorShape& Shape, int32 Axis)
+		{
+			TConstArrayView<uint32> Data = Shape.GetData();
+
+			check(Axis >= 0);
+			check(Axis < Data.Num());
+
+			uint64 Result = 1;
+
+			for (int32 Idx = 0; Idx < Axis; ++Idx)
+			{
+				Result *= Data[Idx];
+			}
+			return Result;
+		}
+
+		uint64 GetNumElemRightOfAxisIncludedForShape(const NNE::FTensorShape& Shape, int32 Axis)
+		{
+			const uint64 NumElemLeftOfAxis = GetNumElemLeftOfAxisForShape(Shape, Axis);
+
+			check(NumElemLeftOfAxis != 0);
+			return Shape.Volume() / NumElemLeftOfAxis;
+		}
+	}
+
+	DECLARE_GPU_STAT_NAMED(FNNEOperatorConcat, TEXT("NNE.Operator.Hlsl.Concat"));
+
 	/**
 	 * Concat operator implementation
 	 */
@@ -27,15 +57,14 @@ namespace UE::NNERuntimeRDG::Private::Hlsl
 			check(OutputTensors.Num() == 1);
 
 			TArray<uint32> OutputShapeData(InputTensors[0]->GetShape().GetData());
-			int32 AxisIndex = Axis >= 0 ? Axis : InputTensors[0]->GetShape().Rank() - Axis;
 			
 			for (int32 i = 1; i < InputTensors.Num(); ++i)
 			{
-				OutputShapeData[AxisIndex] += InputTensors[i]->GetShape().GetData()[AxisIndex];
+				OutputShapeData[Axis] += InputTensors[i]->GetShape().GetData()[Axis];
 				
 				for (int32 r = 0; r < OutputShapeData.Num(); ++r)
 				{
-					if (r != AxisIndex && (OutputShapeData[r] != InputTensors[i]->GetShape().GetData()[r]))
+					if (r != Axis && (OutputShapeData[r] != InputTensors[i]->GetShape().GetData()[r]))
 					{
 						UE_LOG(LogNNE, Warning, TEXT("Concat: all input tensors should have the same shape except on the concatenation axis"));
 						return false;
@@ -49,12 +78,6 @@ namespace UE::NNERuntimeRDG::Private::Hlsl
 
 			Internal::CPUHelper::Concat::Apply(InputTensors, *OutputTensors[0], Axis);
 			
-			if (!OutputTensors[0]->HasPreparedData())
-			{
-				UE_LOG(LogNNE, Warning, TEXT("Concat: Output could not be computed as a constant tensor, however Concat is not implemented on GPU at the moment."));
-				return -1;
-			}
-			
 			return 0;
 		};
 
@@ -64,6 +87,7 @@ namespace UE::NNERuntimeRDG::Private::Hlsl
 			check(OutputTensorDescs.Num() == 1);
 
 			Axis = Attributes.GetValue<int32>(TEXT("axis"));
+			Axis = Axis >= 0 ? Axis : InputTensorDescs[0].GetShape().Rank() + Axis;
 			
 			int32 InputsRank = InputTensorDescs[0].GetShape().Rank();
 			
@@ -85,9 +109,41 @@ namespace UE::NNERuntimeRDG::Private::Hlsl
 			return true;
 		}
 
+		
+
+
 		virtual void Dispatch(FRDGBuilder& GraphBuilder, TConstArrayView<FTensorRDGRef> InputTensors, TConstArrayView<FTensorRDGRef> OutputTensors) override
 		{
-			UE_LOG(LogNNE, Warning, TEXT("Concat: Output should be constant and already uploaded to GPU memory. Dispatch should not need to be called."));
+			check(InputTensors.Num() >= 1);
+			check(OutputTensors.Num() == 1);
+			check(OutputTensors[0] != nullptr);
+
+			const FTensorRDG& Output = *OutputTensors[0];
+			const uint64 NumElemLeftOfAxis = ConcatHelper::GetNumElemLeftOfAxisForShape(Output.GetShape(), Axis);
+			uint64 OutputOffset = 0;
+
+			RDG_EVENT_SCOPE(GraphBuilder, "NNE.Operator.Hlsl.Concat");
+			RDG_GPU_STAT_SCOPE(GraphBuilder, FNNEOperatorConcat);
+
+			for (uint64 IndexShapeLeftOfAxis = 0; IndexShapeLeftOfAxis < NumElemLeftOfAxis; ++IndexShapeLeftOfAxis)
+			{
+				for (FTensorRDGRef InputTensorRef : InputTensors)
+				{
+					check(InputTensorRef != nullptr);
+
+					const FTensorRDG& InputTensor = *InputTensorRef;
+					const uint64 NumElemToCopy = ConcatHelper::GetNumElemRightOfAxisIncludedForShape(InputTensor.GetShape(), Axis);
+					const uint64 FirstElementToCopy = IndexShapeLeftOfAxis * NumElemToCopy;
+					const uint64 NumBytesToCopy = NumElemToCopy * InputTensor.GetElementByteSize();
+					const uint64 FirstByteToCopy = FirstElementToCopy * InputTensor.GetElementByteSize();
+
+					AddCopyBufferPass(GraphBuilder, Output.GetBuffer(), OutputOffset, InputTensor.GetBuffer(), FirstByteToCopy, NumBytesToCopy);
+
+					OutputOffset += NumBytesToCopy;
+				}
+			}
+
+			ensureAlwaysMsgf(OutputOffset == Output.GetDataSize(), TEXT("NNE.Operator.Hlsl.Concat: All of the output buffer should have been written down"));
 		}
 	};
 
diff --git a/Engine/Plugins/Experimental/NNERuntimeRDG/Source/NNERuntimeRDG/Private/Hlsl/NNERuntimeRDGElementWiseBinary.cpp b/Engine/Plugins/Experimental/NNERuntimeRDG/Source/NNERuntimeRDG/Private/Hlsl/NNERuntimeRDGElementWiseBinary.cpp
--- a/Engine/Plugins/Experimental/NNERuntimeRDG/Source/NNERuntimeRDG/Private/Hlsl/NNERuntimeRDGElementWiseBinary.cpp
+++ b/Engine/Plugins/Experimental/NNERuntimeRDG/Source/NNERuntimeRDG/Private/Hlsl/NNERuntimeRDGElementWiseBinary.cpp
@@ -168,7 +168,7 @@ namespace UE::NNERuntimeRDG::Private::Hlsl
 		OP(Prelu);
 		OP(Pow);
 		OP(Sub);
-		//OP(Or);
+		//OP(Xor);
 #undef OP
 
 		return true;
