[NNE] Prevent RDGDml GPU crash and refactor DML temp buffer creation to happen on every frame

Problem
Random crash can happen during inference of FModelInstance

Solution
Make sure that GPU is done initializing the DML resources before dispatching the inference

Tests
NNEIncubator Editor Win64
NNEIncubator Standalone
Compile non-unity

#rnx

diff --git a/Engine/Plugins/Experimental/NNERuntimeRDG/Source/NNERuntimeRDG/Private/Dml/NNEDmlModel.cpp b/Engine/Plugins/Experimental/NNERuntimeRDG/Source/NNERuntimeRDG/Private/Dml/NNEDmlModel.cpp
--- a/Engine/Plugins/Experimental/NNERuntimeRDG/Source/NNERuntimeRDG/Private/Dml/NNEDmlModel.cpp
+++ b/Engine/Plugins/Experimental/NNERuntimeRDG/Source/NNERuntimeRDG/Private/Dml/NNEDmlModel.cpp
@@ -815,6 +815,15 @@ private:
 	int32									NumOutputs;
 };
 
+static constexpr EBufferUsageFlags	WeightBuffUsage = BUF_UnorderedAccess;
+static constexpr ERHIAccess			WeightBuffAccess = ERHIAccess::CopyDest;
+
+static constexpr EBufferUsageFlags	PersistBuffFlags = BUF_Static | BUF_ShaderResource | BUF_UnorderedAccess;
+static constexpr ERHIAccess			PersistBuffAccess = ERHIAccess::UAVMask;
+
+static constexpr EBufferUsageFlags	TempBuffFlags = BUF_Static | BUF_ShaderResource | BUF_UnorderedAccess;
+static constexpr ERHIAccess			TempBuffAccess = ERHIAccess::UAVMask;
+
 FModelInstance::FModelInstance()
 {
 }
@@ -855,15 +864,6 @@ FModelInstance::~FModelInstance()
 				PersistBuff.SafeRelease();
 			}
 
-			if (TempBuff.IsValid())
-			{
-				TempBuff->DisableLifetimeExtension();
-				TempBuff.SafeRelease();
-			}
-
-			RHICmdList.ImmediateFlush(EImmediateFlushType::FlushRHIThreadFlushResources);
-			RHICmdList.ImmediateFlush(EImmediateFlushType::FlushRHIThread);
-
 			Signal->Trigger();
 		}
 	);
@@ -993,15 +993,6 @@ bool FModelInstance::Init(TConstArrayView<uint8> ModelData, FDmlDeviceContext* I
 
 bool FModelInstance::InitCompiledOp()
 {
-	static constexpr EBufferUsageFlags	WeightBuffUsage = BUF_UnorderedAccess;
-	static constexpr ERHIAccess			WeightBuffAccess = ERHIAccess::CopyDest;
-
-	static constexpr EBufferUsageFlags	PersistBuffFlags = BUF_Static | BUF_ShaderResource | BUF_UnorderedAccess;
-	static constexpr ERHIAccess			PersistBuffAccess = ERHIAccess::UAVMask;
-
-	static constexpr EBufferUsageFlags	TempBuffFlags = BUF_Static | BUF_ShaderResource | BUF_UnorderedAccess;
-	static constexpr ERHIAccess			TempBuffAccess = ERHIAccess::UAVMask;
-
 	NNE_TRACE_EVENT_SCOPED(NNE_DmlModelInstance_InitCompiledOp);
 
 	HRESULT					Res;
@@ -1051,7 +1042,7 @@ bool FModelInstance::InitCompiledOp()
 		[
 			this, 
 			Signal, 
-			InitTempMemSize = InitBindProps.TemporaryResourceSize
+			MemSizeInitTemp = InitBindProps.TemporaryResourceSize
 		]
 		(FRHICommandListImmediate& RHICmdList)
 		{
@@ -1070,23 +1061,20 @@ bool FModelInstance::InitCompiledOp()
 			}
 
 			TArray<CD3DX12_RESOURCE_BARRIER, TInlineAllocator<MaxNumInputs>>	Barriers;
-			FGPUFenceRHIRef	UploadFence = nullptr;
 			FBufferRHIRef	UploadBuff;
 
 			if (MemSizeWeights)
 			{
-				UploadFence = RHICreateGPUFence(TEXT("NNE_DmlModelInstance_UploadFence"));
-
 				UploadBuff = CreateRHIBuffer(RHICmdList, MemSizeWeights, BUF_ShaderResource | BUF_Dynamic | BUF_FastVRAM, ERHIAccess::CopySrc, TEXT("NNE_DmlModelInstance_UploadBuffer"));
 				check(UploadBuff);
 				UploadBuff->DisableLifetimeExtension();
 
-				uint8*	UploadBuffPtr = static_cast<uint8*>(RHICmdList.LockBuffer(UploadBuff, 0, MemSizeWeights, RLM_WriteOnly_NoOverwrite));
+				uint8*	UploadBuffPtr = static_cast<uint8*>(RHICmdList.LockBuffer(UploadBuff, 0, MemSizeWeights, RLM_WriteOnly));
 				uint64	UploadOffset = 0;
 				
 				for (int32 WeightIdx = 0; WeightIdx < WeightTensorIndices.Num(); ++WeightIdx)
 				{
-					int32 TensorIdx = WeightTensorIndices[WeightIdx];
+					const int32 TensorIdx = WeightTensorIndices[WeightIdx];
 					if (ConstantCPUTensorIndices.Find(TensorIdx) != -1)
 					{
 						continue;
@@ -1096,9 +1084,11 @@ bool FModelInstance::InitCompiledOp()
 					TConstArrayView<uint8>	TensorData = Tensor.GetPreparedData<uint8>();
 
 					FBufferRHIRef WeightBuff;
+					
 					WeightBuff = CreateRHIBuffer(RHICmdList, TensorData.Num(), WeightBuffUsage, WeightBuffAccess, *Tensor.GetName());
 					check(WeightBuff);
-
+					WeightBuff->DisableLifetimeExtension();
+					
 					FMemory::Memcpy(UploadBuffPtr + UploadOffset, TensorData.GetData(), TensorData.Num());
 					RHICmdList.CopyBufferRegion(WeightBuff, 0, UploadBuff, UploadOffset, TensorData.Num());
 					UploadOffset += Align(TensorData.Num(), DML_MINIMUM_BUFFER_TENSOR_ALIGNMENT);
@@ -1114,7 +1104,6 @@ bool FModelInstance::InitCompiledOp()
 				}
 
 				RHICmdList.UnlockBuffer(UploadBuff);
-				RHICmdList.WriteGPUFence(UploadFence);
 
 				INC_MEMORY_STAT_BY(STAT_MemSizeWeights, MemSizeWeights);
 			}
@@ -1128,34 +1117,21 @@ bool FModelInstance::InitCompiledOp()
 
 			if (MemSizeTemp)
 			{
-				TempBuff = CreateRHIBuffer(RHICmdList, MemSizeTemp, TempBuffFlags, TempBuffAccess, TEXT("NNE_DmlModelInstance_TempBuff"));
-				check(TempBuff.IsValid());
 				INC_MEMORY_STAT_BY(STAT_MemSizeTemp, MemSizeTemp);
 			}
 
 			FBufferRHIRef InitTempBuff;
 
-			if (InitTempMemSize)
+			if (MemSizeInitTemp)
 			{
-				InitTempBuff = CreateRHIBuffer(RHICmdList, InitTempMemSize, TempBuffFlags, TempBuffAccess, TEXT("NNE_DmlModelInstance_InitTempBuff"));
+				InitTempBuff = CreateRHIBuffer(RHICmdList, MemSizeInitTemp, TempBuffFlags, TempBuffAccess, TEXT("NNE_DmlModelInstance_InitTempBuff"));
 				InitTempBuff->DisableLifetimeExtension();
 				check(InitTempBuff.IsValid());
 			}
 
-			if (UploadFence.IsValid())
-			{
-				// Flush commands
-				RHICmdList.ImmediateFlush(EImmediateFlushType::FlushRHIThread);
-
-				while (!UploadFence->Poll())
-				{
-					FPlatformProcess::Sleep(0.0f);
-				}
-			}
-
 			DispatchFence->Clear();
 			RHICmdList.EnqueueLambda(
-				[this, Inputs, Barriers, InitTempBuff](FRHICommandListImmediate& RHICmdList)
+				[this, Inputs = MoveTemp(Inputs), Barriers = MoveTemp(Barriers), InitTempBuff = MoveTemp(InitTempBuff)](FRHICommandListImmediate& RHICmdList)
 				{
 					NNE_TRACE_EVENT_SCOPED(NNE_DmlModelInstance_InitCompiledOpD3D_RHI);
 
@@ -1179,18 +1155,22 @@ bool FModelInstance::InitCompiledOp()
 				}
 			);
 			RHICmdList.WriteGPUFence(DispatchFence);
-
-			// Flush commands
+			
 			RHICmdList.ImmediateFlush(EImmediateFlushType::FlushRHIThread);
 
-			// Now we can release resources
-			if (InitTempBuff.IsValid() || UploadBuff.IsValid())
 			{
-				InitTempBuff.SafeRelease();
-				UploadBuff.SafeRelease();
-
-				RHICmdList.ImmediateFlush(EImmediateFlushType::FlushRHIThreadFlushResources);
-				RHICmdList.ImmediateFlush(EImmediateFlushType::FlushRHIThread);
+				NNE_TRACE_EVENT_SCOPED(NNE_DmlModelInstance_InitCompiledOp_Wait_RT)
+				
+				while (!DispatchFence->Poll())
+				{
+					FPlatformProcess::Sleep(0.0f);
+				}
+				
+				// Potentially we can create a lot of weight buffers, we should release them as soon as possible
+				if (MemSizeWeights)
+				{
+					RHICmdList.ImmediateFlush(EImmediateFlushType::FlushRHIThreadFlushResources);
+				}
 			}
 
 			Signal->Trigger();
@@ -1267,12 +1247,21 @@ void FModelInstance::AddDispatchOps_RenderThread(FRDGBuilder& GraphBuilder)
 				RHIOutputBuffers.Add(RDGBuffer->GetRHI());
 			}
 
+			FBufferRHIRef	TempBuff = nullptr;
+
+			if (MemSizeTemp)
+			{
+				TempBuff = CreateRHIBuffer(RHICmdList, MemSizeTemp, TempBuffFlags, TempBuffAccess, TEXT("NNE_DmlModelInstance_TempBuff"));
+				check(TempBuff.IsValid());
+				TempBuff->DisableLifetimeExtension();
+			}
+
 			{
 				SCOPED_GPU_STAT(RHICmdList, NNE_DmlModelInstance_DispatchD3D_GPU)
 
 				DispatchFence->Clear();
 				RHICmdList.EnqueueLambda(
-					[this, InputBuffers = MoveTemp(RHIInputBuffers), OutputBuffers = MoveTemp(RHIOutputBuffers)](FRHICommandListImmediate& RHICmdList)
+					[this, InputBuffers = MoveTemp(RHIInputBuffers), OutputBuffers = MoveTemp(RHIOutputBuffers), TempBuff = MoveTemp(TempBuff)](FRHICommandListImmediate& RHICmdList)
 					{
 						NNE_TRACE_EVENT_SCOPED(NNE_DmlModelInstance_DispatchD3D_RHI);
 
diff --git a/Engine/Plugins/Experimental/NNERuntimeRDG/Source/NNERuntimeRDG/Private/Dml/NNEDmlModel.h b/Engine/Plugins/Experimental/NNERuntimeRDG/Source/NNERuntimeRDG/Private/Dml/NNEDmlModel.h
--- a/Engine/Plugins/Experimental/NNERuntimeRDG/Source/NNERuntimeRDG/Private/Dml/NNEDmlModel.h
+++ b/Engine/Plugins/Experimental/NNERuntimeRDG/Source/NNERuntimeRDG/Private/Dml/NNEDmlModel.h
@@ -120,7 +120,6 @@ private:
 
 	FGPUFenceRHIRef						DispatchFence;
 	FBufferRHIRef						PersistBuff;
-	FBufferRHIRef						TempBuff;
 	uint64								MemSizeWeights;
 	uint64								MemSizeTemp;
 	uint64								MemSizePersist;
