[NNE] NNERuntimeRDGHlsl transpose implementation

Problem:
NNERuntimeRDGHlsl was missing transpose.

Solution:
Implement it

Test:
Win64 Editor: transpose test + smoke tests
Win64 Standalone: all tests

#rnx

diff --git a/Engine/Plugins/Experimental/NNERuntimeRDG/Source/NNEHlslShaders/Internal/NNEHlslShadersTransposeCS.h b/Engine/Plugins/Experimental/NNERuntimeRDG/Source/NNEHlslShaders/Internal/NNEHlslShadersTransposeCS.h
new file mode 100644
--- /dev/null
+++ b/Engine/Plugins/Experimental/NNERuntimeRDG/Source/NNEHlslShaders/Internal/NNEHlslShadersTransposeCS.h
@@ -0,0 +1,39 @@
+// Copyright Epic Games, Inc. All Rights Reserved.
+
+#pragma once
+
+#include "NNEHlslShaderBase.h"
+#include "RenderGraphUtils.h"
+#include "ShaderParameterUtils.h"
+
+namespace UE::NNEHlslShaders::Internal
+{
+	class FTransposeConstants
+	{
+	public:
+
+		static const int32 MAX_NUM_DIMENSIONS{ 8 };
+		static const int32 NUM_GROUP_THREADS{ 256 };
+	};
+
+	class NNEHLSLSHADERS_API FTransposeCS : public FHlslShaderBase
+	{
+		DECLARE_GLOBAL_SHADER(FTransposeCS);
+		SHADER_USE_PARAMETER_STRUCT(FTransposeCS, FHlslShaderBase)
+
+		class FTransposeNumDimensions : SHADER_PERMUTATION_RANGE_INT("NUM_DIMENSIONS", 1, FTransposeConstants::MAX_NUM_DIMENSIONS);
+		using FPermutationDomain = TShaderPermutationDomain<FTransposeNumDimensions>;
+
+	public:
+
+		BEGIN_SHADER_PARAMETER_STRUCT(FParameters, )
+			SHADER_PARAMETER_RDG_BUFFER_SRV(Buffer<float>, Input)
+			SHADER_PARAMETER_RDG_BUFFER_UAV(RWBuffer<float>, Output)
+			SHADER_PARAMETER_ARRAY(FUintVector4, TensorInfo, [FTransposeConstants::MAX_NUM_DIMENSIONS])
+			SHADER_PARAMETER(uint32, Num)
+			SHADER_PARAMETER(uint32, ThreadCountX)
+		END_SHADER_PARAMETER_STRUCT()
+
+		static void ModifyCompilationEnvironment(const FGlobalShaderPermutationParameters& InParameters, FShaderCompilerEnvironment& OutEnvironment);
+	};
+} // UE::NNEHlslShaders::Internal
diff --git a/Engine/Plugins/Experimental/NNERuntimeRDG/Source/NNEHlslShaders/Private/NNEHlslShadersTransposeCS.cpp b/Engine/Plugins/Experimental/NNERuntimeRDG/Source/NNEHlslShaders/Private/NNEHlslShadersTransposeCS.cpp
new file mode 100644
--- /dev/null
+++ b/Engine/Plugins/Experimental/NNERuntimeRDG/Source/NNEHlslShaders/Private/NNEHlslShadersTransposeCS.cpp
@@ -0,0 +1,18 @@
+// Copyright Epic Games, Inc. All Rights Reserved.
+
+#include "NNEHlslShadersTransposeCS.h"
+#include "NNE.h"
+
+namespace UE::NNEHlslShaders::Internal
+{
+	void FTransposeCS::ModifyCompilationEnvironment(const FGlobalShaderPermutationParameters& InParameters, FShaderCompilerEnvironment& OutEnvironment)
+	{
+		FGlobalShader::ModifyCompilationEnvironment(InParameters, OutEnvironment);
+
+		OutEnvironment.SetDefine(TEXT("THREADGROUP_SIZE_X"), FTransposeConstants::NUM_GROUP_THREADS);
+
+		FPermutationDomain PermutationVector(InParameters.PermutationId);
+	}
+
+	IMPLEMENT_GLOBAL_SHADER(FTransposeCS, "/NNE/NNEHlslShadersTranspose.usf", "Transpose", SF_Compute);
+} // UE::NNEHlslShaders::Internal
\ No newline at end of file
diff --git a/Engine/Plugins/Experimental/NNERuntimeRDG/Source/NNEHlslShaders/Shaders/NNEHlslShadersTranspose.usf b/Engine/Plugins/Experimental/NNERuntimeRDG/Source/NNEHlslShaders/Shaders/NNEHlslShadersTranspose.usf
new file mode 100644
--- /dev/null
+++ b/Engine/Plugins/Experimental/NNERuntimeRDG/Source/NNEHlslShaders/Shaders/NNEHlslShadersTranspose.usf
@@ -0,0 +1,36 @@
+// Copyright Epic Games, Inc. All Rights Reserved.
+
+#include "/Engine/Public/Platform.ush"
+#include "NNEHlslShadersBroadcastHelper.ush"
+
+Buffer<float> Input;
+RWBuffer<float> Output;
+uint4 TensorInfo[NUM_DIMENSIONS];
+uint Num;
+uint ThreadCountX;
+
+#define OUTPUT_STRIDE 0
+#define INPUT_STRIDE_TRANSPOSED 2
+
+[numthreads(THREADGROUP_SIZE_X, 1, 1)]
+void Transpose(in const uint3 DispatchThreadID : SV_DispatchThreadID)
+{
+	const uint Index = DispatchThreadID.y * ThreadCountX + DispatchThreadID.x;
+	
+	if (Index < Num)
+	{
+		uint InputIndex = 0;
+		uint Offset = Index;
+
+		for (uint dim = 0; dim < NUM_DIMENSIONS; ++dim)
+		{
+			uint OutDimIdx, R;
+			DivMod(Offset, TensorInfo[dim][OUTPUT_STRIDE], OutDimIdx, R);
+			InputIndex += TensorInfo[dim][INPUT_STRIDE_TRANSPOSED] * OutDimIdx;
+			Offset = R;
+		}
+
+		float X = Input[InputIndex];
+		Output[Index] = X;
+	}
+}
diff --git a/Engine/Plugins/Experimental/NNERuntimeRDG/Source/NNERuntimeRDG/Internal/NNERuntimeRDGHelperTranspose.h b/Engine/Plugins/Experimental/NNERuntimeRDG/Source/NNERuntimeRDG/Internal/NNERuntimeRDGHelperTranspose.h
--- a/Engine/Plugins/Experimental/NNERuntimeRDG/Source/NNERuntimeRDG/Internal/NNERuntimeRDGHelperTranspose.h
+++ b/Engine/Plugins/Experimental/NNERuntimeRDG/Source/NNERuntimeRDG/Internal/NNERuntimeRDGHelperTranspose.h
@@ -8,5 +8,6 @@ namespace UE::NNE::Internal { class FTensor; }
 
 namespace UE::NNERuntimeRDG::Internal::CPUHelper::Transpose
 {
+	bool NNERUNTIMERDG_API Apply(const NNE::Internal::FTensor& DataTensor, TConstArrayView<int32> Perms, NNE::Internal::FTensor& OutputTensor);
 	bool NNERUNTIMERDG_API TransposePreparedData(NNE::Internal::FTensor& Tensor, TConstArrayView<int32> Perms);
 } // UE::NNERuntimeRDG::Internal::CPUHelper::Transpose
\ No newline at end of file
diff --git a/Engine/Plugins/Experimental/NNERuntimeRDG/Source/NNERuntimeRDG/Private/Helper/NNERuntimeRDGHelperTranspose.cpp b/Engine/Plugins/Experimental/NNERuntimeRDG/Source/NNERuntimeRDG/Private/Helper/NNERuntimeRDGHelperTranspose.cpp
--- a/Engine/Plugins/Experimental/NNERuntimeRDG/Source/NNERuntimeRDG/Private/Helper/NNERuntimeRDGHelperTranspose.cpp
+++ b/Engine/Plugins/Experimental/NNERuntimeRDG/Source/NNERuntimeRDG/Private/Helper/NNERuntimeRDGHelperTranspose.cpp
@@ -5,7 +5,7 @@
 
 namespace UE::NNERuntimeRDG::Internal::CPUHelper::Transpose
 {
-	bool TransposePreparedData(NNE::Internal::FTensor& Tensor, TConstArrayView<int32> Perms)
+	bool Apply(const NNE::Internal::FTensor& Tensor, TConstArrayView<int32> Perms, NNE::Internal::FTensor& OutputTensor)
 	{
 		if (!Tensor.HasPreparedData())
 		{
@@ -50,9 +50,15 @@ namespace UE::NNERuntimeRDG::Internal::CPUHelper::Transpose
 
 		} while (it.Advance());
 
-		Tensor.SetPreparedData<float>(TransposedData);
+		OutputTensor.SetPreparedData<float>(TransposedData);
 
 		return true;
 	}
+
+	bool TransposePreparedData(NNE::Internal::FTensor& Tensor, TConstArrayView<int32> Perms)
+	{
+		//Transpose in place
+		return Apply(Tensor, Perms, Tensor);
+	}
 	
 } // UE::NNERuntimeRDG::Private::CPUHelper::Transpose
diff --git a/Engine/Plugins/Experimental/NNERuntimeRDG/Source/NNERuntimeRDG/Private/Hlsl/NNERuntimeRDGTranspose.cpp b/Engine/Plugins/Experimental/NNERuntimeRDG/Source/NNERuntimeRDG/Private/Hlsl/NNERuntimeRDGTranspose.cpp
new file mode 100644
--- /dev/null
+++ b/Engine/Plugins/Experimental/NNERuntimeRDG/Source/NNERuntimeRDG/Private/Hlsl/NNERuntimeRDGTranspose.cpp
@@ -0,0 +1,159 @@
+// Copyright Epic Games, Inc. All Rights Reserved.
+
+#include "NNERuntimeRDGTranspose.h"
+#include "NNEHlslShadersTransposeCS.h"
+#include "NNERuntimeRDGHelperTranspose.h"
+#include "NNERuntimeRDGHlslHelper.h"
+#include "NNETensor.h"
+#include "NNETypes.h"
+#include "RenderGraphUtils.h"
+
+namespace UE::NNERuntimeRDG::Private::Hlsl
+{
+	DECLARE_GPU_STAT_NAMED(FNNEOperatorTranspose, TEXT("NNE.Operator.Hlsl.Transpose"));
+
+	/**
+	 * Transpose operator implementation
+	 */
+	class FTranspose : public FOperatorHlsl
+	{
+	public:
+
+		FTranspose() {}
+		virtual ~FTranspose() = default;
+
+	private:
+		TArray<int32> Perm;
+
+	public:
+
+		virtual int PrepareOutputs(TConstArrayView<NNE::Internal::FTensorRef> InputTensors, TArrayView<NNE::Internal::FTensorRef> OutputTensors) const override
+		{
+			check(InputTensors.Num() == 1);
+			check(OutputTensors.Num() == 1);
+
+			const NNE::Internal::FTensor& X = *InputTensors[0];
+			TConstArrayView<uint32> InputShape = X.GetShape().GetData();
+			TArray<uint32, TInlineAllocator<NNE::FTensorShape::MaxRank>> OutputShape;
+
+			OutputShape.SetNumUninitialized(InputShape.Num());
+			for (int32 i = 0; i < InputShape.Num(); ++i)
+			{
+				OutputShape[i] = InputShape[Perm[i]];
+			}
+			OutputTensors[0]->SetShape(NNE::FTensorShape::Make(OutputShape));
+
+			Internal::CPUHelper::Transpose::Apply(*InputTensors[0], Perm, *OutputTensors[0]);
+
+			return 0;
+		};
+
+		virtual bool Initialize(TConstArrayView<NNE::FTensorDesc> InputTensorDescs, TConstArrayView<NNE::FTensorDesc> OutputTensorDescs, const NNE::FAttributeMap& Attributes) override
+		{
+			check(InputTensorDescs.Num() == 1);
+			check(OutputTensorDescs.Num() == 1);
+
+			const NNE::FTensorDesc& InputDesc = InputTensorDescs[0];
+			const int32 InputRank = InputDesc.GetShape().Rank();
+			TArray<int32> PermutationDefault;
+
+			PermutationDefault.Reserve(NNE::FTensorShape::MaxRank);
+			for (int32 i = 0; i < InputRank; ++i)
+			{
+				PermutationDefault.Add((InputRank -1) - i);
+			}
+			Perm = Attributes.GetValueOrDefault<TArray<int32>>(TEXT("perm"), PermutationDefault);
+
+			if (Perm.Num() != InputTensorDescs[0].GetShape().Rank())
+			{
+				UE_LOG(LogNNE, Warning, TEXT("Transpose 'perm' attribute should contain the same amount of element as the rank of the input tensor."));
+				return false;
+			}
+
+			return true;
+		}
+
+		virtual void Dispatch(FRDGBuilder& GraphBuilder, TConstArrayView<FTensorRDGRef> InputTensors, TConstArrayView<FTensorRDGRef> OutputTensors) override
+		{
+			using namespace UE::NNEHlslShaders::Internal;
+
+			check(InputTensors.Num() == 1);
+			check(OutputTensors.Num() == 1);
+			check(InputTensors[0] != nullptr);
+			check(OutputTensors[0] != nullptr);
+
+			const FTensorRDG& Input = *InputTensors[0];
+			const FTensorRDG& Output = *OutputTensors[0];
+			const FRDGBufferSRVRef InputSRV = GraphBuilder.CreateSRV(FRDGBufferSRVDesc(Input.GetBuffer(), PF_R32_FLOAT));
+			const FRDGBufferUAVRef OutputUAV = GraphBuilder.CreateUAV(FRDGBufferUAVDesc(Output.GetBuffer(), PF_R32_FLOAT));
+			const FIntVector ThreadGroupCount = ComputeElementWiseThreadGroups(Output.GetVolume(), FTransposeConstants::NUM_GROUP_THREADS);
+
+			// Set parameters
+			FTransposeCS::FParameters* Params = GraphBuilder.AllocParameters<FTransposeCS::FParameters>();
+			Params->Input = InputSRV;
+			Params->Output = OutputUAV;
+			Params->Num = Output.GetVolume();
+			Params->ThreadCountX = ThreadGroupCount.X * FTransposeConstants::NUM_GROUP_THREADS;
+			FillTensorStrideShaderParameters(Output, Params->TensorInfo, 0);
+			FillTensorStrideShaderParameters(Input, Params->TensorInfo, 1);
+			static_assert(NNE::FTensorShape::MaxRank <= NXRT_TENSORSTRIDEINFO_MAX_NUM_DIMENSIONS);
+			check(Perm.Num() == Input.GetShape().Rank());
+			for (int32 i = 0; i < NXRT_TENSORSTRIDEINFO_MAX_NUM_DIMENSIONS; ++i)
+			{
+				if (i < Input.GetShape().Rank())
+				{
+					Params->TensorInfo[i][2] = Params->TensorInfo[Perm[i]][1];
+				}
+				else
+				{
+					Params->TensorInfo[i][2] = 0;
+				}
+			}
+
+			FTransposeCS::FPermutationDomain PermutationVector;
+			PermutationVector.Set<FTransposeCS::FTransposeNumDimensions>(Output.GetShape().Rank());
+
+			TShaderMapRef<FTransposeCS> ComputeShader(GetGlobalShaderMap(GMaxRHIFeatureLevel), PermutationVector);
+
+			RDG_EVENT_SCOPE(GraphBuilder, "NNE.Operator.Hlsl.Transpose");
+			RDG_GPU_STAT_SCOPE(GraphBuilder, FNNEOperatorTranspose);
+
+			FComputeShaderUtils::AddPass(
+				GraphBuilder,
+				RDG_EVENT_NAME("NNE.Operator.Hlsl.Transpose.Dispatch"),
+				ERDGPassFlags::Compute | ERDGPassFlags::NeverCull,
+				ComputeShader,
+				Params,
+				ThreadGroupCount);
+		}
+	};
+
+	bool ValidateTransposeOperator(const NNE::FAttributeMap& AttributeMap, TConstArrayView<ENNETensorDataType> InputTypes, TConstArrayView<NNE::FSymbolicTensorShape> InputTransposes)
+	{
+		bool bIsValid = true;
+
+		//This match version 1 and 13 of the Transpose operator
+		//[REDACTED_PROFILE_URL]/onnx/blob/main/docs/Operators.md#Transpose
+		FAttributeValidator AttributeValidator;
+		AttributeValidator.AddOptional(TEXT("perm"), ENNEAttributeDataType::Int32Array);
+		bIsValid &= AttributeValidator.Validate(AttributeMap);
+
+		FInputValidator InputValidator;
+		InputValidator.AddSupportedType(ENNETensorDataType::Float);
+		InputValidator.AddRequired();
+		bIsValid &= InputValidator.Validate(InputTypes);
+
+		return bIsValid;
+	}
+
+	FOperatorHlsl* CreateTransposeOperator()
+	{
+		return new FTranspose();
+	}
+
+	bool RegisterTransposeOperator(FOperatorRegistryHlsl& Registry)
+	{
+		Registry.OpAdd(TEXT("Transpose"), CreateTransposeOperator, ValidateTransposeOperator);
+		return true;
+	}
+} // UE::NNERuntimeRDG::Private::Hlsl
diff --git a/Engine/Plugins/Experimental/NNERuntimeRDG/Source/NNERuntimeRDG/Private/Hlsl/NNERuntimeRDGTranspose.h b/Engine/Plugins/Experimental/NNERuntimeRDG/Source/NNERuntimeRDG/Private/Hlsl/NNERuntimeRDGTranspose.h
new file mode 100644
--- /dev/null
+++ b/Engine/Plugins/Experimental/NNERuntimeRDG/Source/NNERuntimeRDG/Private/Hlsl/NNERuntimeRDGTranspose.h
@@ -0,0 +1,11 @@
+// Copyright Epic Games, Inc. All Rights Reserved.
+
+#pragma once
+
+#include "NNERuntimeRDGBase.h"
+#include "NNERuntimeRDGHlslOp.h"
+
+namespace UE::NNERuntimeRDG::Private::Hlsl
+{
+	bool RegisterTransposeOperator(FOperatorRegistryHlsl& Registry);
+} // UE::NNERuntimeRDG::Private::Hlsl
diff --git a/Engine/Plugins/Experimental/NNERuntimeRDG/Source/NNERuntimeRDG/Private/NNERuntimeRDGHlsl.cpp b/Engine/Plugins/Experimental/NNERuntimeRDG/Source/NNERuntimeRDG/Private/NNERuntimeRDGHlsl.cpp
--- a/Engine/Plugins/Experimental/NNERuntimeRDG/Source/NNERuntimeRDG/Private/NNERuntimeRDGHlsl.cpp
+++ b/Engine/Plugins/Experimental/NNERuntimeRDG/Source/NNERuntimeRDG/Private/NNERuntimeRDGHlsl.cpp
@@ -34,6 +34,7 @@
 #include "Hlsl/NNERuntimeRDGSize.h"
 #include "Hlsl/NNERuntimeRDGSlice.h"
 #include "Hlsl/NNERuntimeRDGSqueeze.h"
+#include "Hlsl/NNERuntimeRDGTranspose.h"
 #include "Hlsl/NNERuntimeRDGUnsqueeze.h"
 #include "Hlsl/NNERuntimeRDGUpsample.h"
 #include "Hlsl/NNERuntimeRDGMatMul.h"
@@ -68,6 +69,7 @@ bool UNNERuntimeRDGHlslImpl::Init()
 	RegisterSizeOperator(*registry);
 	RegisterSliceOperator(*registry);
 	RegisterSqueezeOperator(*registry);
+	RegisterTransposeOperator(*registry);
 	RegisterUnsqueezeOperator(*registry);
 	RegisterUpsampleOperator(*registry);
 	RegisterMatMulOperator(*registry);
