[NNE] NNERuntimeRDGHlsl Slice GPU implementation

Problem:
NNERuntimeRDGHlsl was missing slice gpu implementation

Solution:
Implement it

Test:
Win64 Editor: all tests
Win64 Standalone: all tests

#rnx

diff --git a/Engine/Plugins/Experimental/NNERuntimeRDG/Source/NNEHlslShaders/Internal/NNEHlslShadersSliceCS.h b/Engine/Plugins/Experimental/NNERuntimeRDG/Source/NNEHlslShaders/Internal/NNEHlslShadersSliceCS.h
new file mode 100644
--- /dev/null
+++ b/Engine/Plugins/Experimental/NNERuntimeRDG/Source/NNEHlslShaders/Internal/NNEHlslShadersSliceCS.h
@@ -0,0 +1,39 @@
+// Copyright Epic Games, Inc. All Rights Reserved.
+
+#pragma once
+
+#include "NNEHlslShaderBase.h"
+#include "RenderGraphUtils.h"
+#include "ShaderParameterUtils.h"
+
+namespace UE::NNEHlslShaders::Internal
+{
+	class FSliceConstants
+	{
+	public:
+
+		static const int32 MAX_NUM_DIMENSIONS{ 8 };
+		static const int32 NUM_GROUP_THREADS{ 256 };
+	};
+
+	class NNEHLSLSHADERS_API FSliceCS : public FHlslShaderBase
+	{
+		DECLARE_GLOBAL_SHADER(FSliceCS);
+		SHADER_USE_PARAMETER_STRUCT(FSliceCS, FHlslShaderBase)
+
+		class FSliceNumDimensions : SHADER_PERMUTATION_RANGE_INT("NUM_DIMENSIONS", 1, FSliceConstants::MAX_NUM_DIMENSIONS);
+		using FPermutationDomain = TShaderPermutationDomain<FSliceNumDimensions>;
+
+	public:
+
+		BEGIN_SHADER_PARAMETER_STRUCT(FParameters, )
+			SHADER_PARAMETER_RDG_BUFFER_SRV(Buffer<float>, Input)
+			SHADER_PARAMETER_RDG_BUFFER_UAV(RWBuffer<float>, Output)
+			SHADER_PARAMETER_ARRAY(FUintVector4, TensorInfo, [FSliceConstants::MAX_NUM_DIMENSIONS])
+			SHADER_PARAMETER(uint32, Num)
+			SHADER_PARAMETER(uint32, ThreadCountX)
+		END_SHADER_PARAMETER_STRUCT()
+
+		static void ModifyCompilationEnvironment(const FGlobalShaderPermutationParameters& InParameters, FShaderCompilerEnvironment& OutEnvironment);
+	};
+} // UE::NNEHlslShaders::Internal
diff --git a/Engine/Plugins/Experimental/NNERuntimeRDG/Source/NNEHlslShaders/Private/NNEHlslShadersSliceCS.cpp b/Engine/Plugins/Experimental/NNERuntimeRDG/Source/NNEHlslShaders/Private/NNEHlslShadersSliceCS.cpp
new file mode 100644
--- /dev/null
+++ b/Engine/Plugins/Experimental/NNERuntimeRDG/Source/NNEHlslShaders/Private/NNEHlslShadersSliceCS.cpp
@@ -0,0 +1,18 @@
+// Copyright Epic Games, Inc. All Rights Reserved.
+
+#include "NNEHlslShadersSliceCS.h"
+#include "NNE.h"
+
+namespace UE::NNEHlslShaders::Internal
+{
+	void FSliceCS::ModifyCompilationEnvironment(const FGlobalShaderPermutationParameters& InParameters, FShaderCompilerEnvironment& OutEnvironment)
+	{
+		FGlobalShader::ModifyCompilationEnvironment(InParameters, OutEnvironment);
+
+		OutEnvironment.SetDefine(TEXT("THREADGROUP_SIZE_X"), FSliceConstants::NUM_GROUP_THREADS);
+
+		FPermutationDomain PermutationVector(InParameters.PermutationId);
+	}
+
+	IMPLEMENT_GLOBAL_SHADER(FSliceCS, "/NNE/NNEHlslShadersSlice.usf", "Slice", SF_Compute);
+} // UE::NNEHlslShaders::Internal
\ No newline at end of file
diff --git a/Engine/Plugins/Experimental/NNERuntimeRDG/Source/NNEHlslShaders/Shaders/NNEHlslShadersSlice.usf b/Engine/Plugins/Experimental/NNERuntimeRDG/Source/NNEHlslShaders/Shaders/NNEHlslShadersSlice.usf
new file mode 100644
--- /dev/null
+++ b/Engine/Plugins/Experimental/NNERuntimeRDG/Source/NNEHlslShaders/Shaders/NNEHlslShadersSlice.usf
@@ -0,0 +1,38 @@
+// Copyright Epic Games, Inc. All Rights Reserved.
+
+#include "/Engine/Public/Platform.ush"
+#include "NNEHlslShadersBroadcastHelper.ush"
+
+Buffer<float> Input;
+RWBuffer<float> Output;
+uint4 TensorInfo[NUM_DIMENSIONS];
+uint Num;
+uint ThreadCountX;
+
+#define INPUT_STRIDE 0
+#define OUTPUT_STRIDE 1
+#define INPUT_STARTS 2
+
+[numthreads(THREADGROUP_SIZE_X, 1, 1)]
+void Slice(in const uint3 DispatchThreadID : SV_DispatchThreadID)
+{
+	const uint Index = DispatchThreadID.y * ThreadCountX + DispatchThreadID.x;
+	
+	if (Index < Num)
+	{
+		uint InputIndex = 0;
+		uint Offset = Index;
+
+		for (uint dim = 0; dim < NUM_DIMENSIONS; ++dim)
+		{
+			uint OutDimIdx, R;
+			DivMod(Offset, TensorInfo[dim][OUTPUT_STRIDE], OutDimIdx, R);
+			uint InDimIdx = TensorInfo[dim][INPUT_STARTS] + OutDimIdx;
+			InputIndex += TensorInfo[dim][INPUT_STRIDE] * InDimIdx;
+			Offset = R;
+		}
+
+		float X = Input[InputIndex];
+		Output[Index] = X;
+	}
+}
diff --git a/Engine/Plugins/Experimental/NNERuntimeRDG/Source/NNERuntimeRDG/Private/Hlsl/NNERuntimeRDGSlice.cpp b/Engine/Plugins/Experimental/NNERuntimeRDG/Source/NNERuntimeRDG/Private/Hlsl/NNERuntimeRDGSlice.cpp
--- a/Engine/Plugins/Experimental/NNERuntimeRDG/Source/NNERuntimeRDG/Private/Hlsl/NNERuntimeRDGSlice.cpp
+++ b/Engine/Plugins/Experimental/NNERuntimeRDG/Source/NNERuntimeRDG/Private/Hlsl/NNERuntimeRDGSlice.cpp
@@ -1,12 +1,17 @@
 // Copyright Epic Games, Inc. All Rights Reserved.
 
 #include "NNERuntimeRDGSlice.h"
+#include "NNEHlslShadersSliceCS.h"
 #include "NNERuntimeRDGHelperSlice.h"
+#include "NNERuntimeRDGHlslHelper.h"
 #include "NNETensor.h"
 #include "NNETypes.h"
+#include "RenderGraphUtils.h"
 
 namespace UE::NNERuntimeRDG::Private::Hlsl
 {
+	DECLARE_GPU_STAT_NAMED(FNNEOperatorSlice, TEXT("NNE.Operator.Hlsl.Slice"));
+
 	/**
 	 * Slice operator implementation
 	 */
@@ -17,25 +22,20 @@ namespace UE::NNERuntimeRDG::Private::Hlsl
 		FSlice() {}
 		virtual ~FSlice() = default;
 
+	private:
+
 		TArray<int32> AxesAttr;
 		TArray<int32> EndsAttr;
 		TArray<int32> StartsAttr;
 
-	public:
-
-		virtual int PrepareOutputs(TConstArrayView<NNE::Internal::FTensorRef> InputTensors, TArrayView<NNE::Internal::FTensorRef> OutputTensors) const override
+		inline void GetStartAndEndFromInputShape(TConstArrayView<uint32> InputShapeData, 
+			TArray<int32, TInlineAllocator<NNE::FTensorShape::MaxRank>>& Start,
+			TArray<int32, TInlineAllocator<NNE::FTensorShape::MaxRank>>& End) const
 		{
-			check(InputTensors.Num() == 1);
-			check(OutputTensors.Num() == 1);
-			
-			TConstArrayView<uint32> Dims(InputTensors[0]->GetShape().GetData());
-			int32 InputRank = Dims.Num();
-
-			TArray<int32> Axes = AxesAttr;
-			TArray<int32> Ends = EndsAttr;
-			TArray<int32> Starts = StartsAttr;
-			TArray<int32> Start;
-			TArray<int32> End;
+			const int32 InputRank = InputShapeData.Num();
+			TArray<int32, TInlineAllocator<NNE::FTensorShape::MaxRank>> Axes(AxesAttr);
+			TArray<int32, TInlineAllocator<NNE::FTensorShape::MaxRank>> Ends(EndsAttr);
+			TArray<int32, TInlineAllocator<NNE::FTensorShape::MaxRank>> Starts(StartsAttr);
 
 			//see [REDACTED_PROFILE_URL]/onnx/blob/main/docs/Operators.md#slice for algorithm
 			Start.SetNum(InputRank);
@@ -43,7 +43,7 @@ namespace UE::NNERuntimeRDG::Private::Hlsl
 			for (int32 i = 0; i < InputRank; ++i)
 			{
 				Start[i] = 0;
-				End[i] = Dims[i];
+				End[i] = InputShapeData[i];
 			}
 			for (int32 i = 0; i < Axes.Num(); ++i)
 			{
@@ -56,24 +56,39 @@ namespace UE::NNERuntimeRDG::Private::Hlsl
 			{
 				if (Starts[i] < 0)
 				{
-					Starts[i] += Dims[Axes[i]];
+					Starts[i] += InputShapeData[Axes[i]];
 				}
 			}
 			for (int32 i = 0; i < Ends.Num(); ++i)
 			{
 				if (Ends[i] < 0)
 				{
-					Ends[i] += Dims[Axes[i]];
+					Ends[i] += InputShapeData[Axes[i]];
 				}
 			}
 			for (int32 i = 0; i < Axes.Num(); ++i)
 			{
-				Start[Axes[i]] = FMath::Clamp(Starts[i], 0, Dims[Axes[i]]);
-				End[Axes[i]] = FMath::Clamp(Ends[i], 0, Dims[Axes[i]]);
+				Start[Axes[i]] = FMath::Clamp(Starts[i], 0, InputShapeData[Axes[i]]);
+				End[Axes[i]] = FMath::Clamp(Ends[i], 0, InputShapeData[Axes[i]]);
 			}
+		}
+
+	public:
+
+		virtual int PrepareOutputs(TConstArrayView<NNE::Internal::FTensorRef> InputTensors, TArrayView<NNE::Internal::FTensorRef> OutputTensors) const override
+		{
+			check(InputTensors.Num() == 1);
+			check(OutputTensors.Num() == 1);
+
+			TConstArrayView<uint32> InputShapeData(InputTensors[0]->GetShape().GetData());
+			const int32 InputRank = InputShapeData.Num();
+			TArray<int32, TInlineAllocator<NNE::FTensorShape::MaxRank>> Start(AxesAttr);
+			TArray<int32, TInlineAllocator<NNE::FTensorShape::MaxRank>> End(EndsAttr);
+
+			GetStartAndEndFromInputShape(InputShapeData, Start, End);
 
 			TArray<uint32> OutputShapeData;
-			
+
 			OutputShapeData.Reserve(InputRank);
 			for (int32 i = 0; i < InputRank; ++i)
 			{
@@ -86,12 +101,6 @@ namespace UE::NNERuntimeRDG::Private::Hlsl
 			
 			Internal::CPUHelper::Slice::Apply(*InputTensors[0], *OutputTensors[0], Start);
 
-			if (!OutputTensors[0]->HasPreparedData())
-			{
-				UE_LOG(LogNNE, Warning, TEXT("Slice: Output could not be computed as a constant tensor, however Slice is not implemented on GPU at the moment."));
-				return -1;
-			}
-
 			return 0;
 		};
 
@@ -123,7 +132,62 @@ namespace UE::NNERuntimeRDG::Private::Hlsl
 
 		virtual void Dispatch(FRDGBuilder& GraphBuilder, TConstArrayView<FTensorRDGRef> InputTensors, TConstArrayView<FTensorRDGRef> OutputTensors) override
 		{
-			UE_LOG(LogNNE, Warning, TEXT("Slice: Output should be constant and already uploaded to GPU memory. Dispatch should not need to be called."));
+			using namespace UE::NNEHlslShaders::Internal;
+
+			check(InputTensors.Num() == 1);
+			check(OutputTensors.Num() == 1);
+			check(InputTensors[0] != nullptr);
+			check(OutputTensors[0] != nullptr);
+
+			const FTensorRDG& Input = *InputTensors[0];
+			const FTensorRDG& Output = *OutputTensors[0];
+			const FRDGBufferSRVRef InputSRV = GraphBuilder.CreateSRV(FRDGBufferSRVDesc(Input.GetBuffer(), PF_R32_FLOAT));
+			const FRDGBufferUAVRef OutputUAV = GraphBuilder.CreateUAV(FRDGBufferUAVDesc(Output.GetBuffer(), PF_R32_FLOAT));
+			const FIntVector ThreadGroupCount = ComputeElementWiseThreadGroups(Output.GetVolume(), FSliceConstants::NUM_GROUP_THREADS);
+
+			TConstArrayView<uint32> InputShapeData(Input.GetShape().GetData());
+			TArray<int32, TInlineAllocator<NNE::FTensorShape::MaxRank>> Start(AxesAttr);
+			TArray<int32, TInlineAllocator<NNE::FTensorShape::MaxRank>> End(EndsAttr);
+
+			GetStartAndEndFromInputShape(InputShapeData, Start, End);
+
+			// Set parameters
+			FSliceCS::FParameters* Params = GraphBuilder.AllocParameters<FSliceCS::FParameters>();
+			Params->Input = InputSRV;
+			Params->Output = OutputUAV;
+			Params->Num = Output.GetVolume();
+			Params->ThreadCountX = ThreadGroupCount.X * FSliceConstants::NUM_GROUP_THREADS;
+			FillTensorStrideShaderParameters(Input, Params->TensorInfo, 0);
+			FillTensorStrideShaderParameters(Output, Params->TensorInfo, 1);
+			static_assert(NNE::FTensorShape::MaxRank <= NXRT_TENSORSTRIDEINFO_MAX_NUM_DIMENSIONS);
+			check(Start.Num() == Input.GetShape().Rank());
+			for (int32 i = 0; i < NXRT_TENSORSTRIDEINFO_MAX_NUM_DIMENSIONS; ++i)
+			{
+				if (i < Start.Num())
+				{
+					Params->TensorInfo[i][2] = static_cast<uint32>(Start[i]);
+				}
+				else
+				{
+					Params->TensorInfo[i][2] = 0;
+				}
+			}
+
+			FSliceCS::FPermutationDomain PermutationVector;
+			PermutationVector.Set<FSliceCS::FSliceNumDimensions>(Output.GetShape().Rank());
+
+			TShaderMapRef<FSliceCS> ComputeShader(GetGlobalShaderMap(GMaxRHIFeatureLevel), PermutationVector);
+
+			RDG_EVENT_SCOPE(GraphBuilder, "NNE.Operator.Hlsl.Slice");
+			RDG_GPU_STAT_SCOPE(GraphBuilder, FNNEOperatorSlice);
+
+			FComputeShaderUtils::AddPass(
+				GraphBuilder,
+				RDG_EVENT_NAME("NNE.Operator.Hlsl.Slice.Dispatch"),
+				ERDGPassFlags::Compute | ERDGPassFlags::NeverCull,
+				ComputeShader,
+				Params,
+				ThreadGroupCount);
 		}
 	};
 
