[NNE] NNERuntimeRDGHlsl Conv.Matmul fix

Problem:
When batch > 1 and Ci != Cw Conv.Matmul implementation was producing incorrect results

Solution:
Fixed batch offset in impl, added related tests.

Test:
Win64 Editor: conv tests + style transfer tests + smoke tests
Win64 Standalone: all tests

#rnx

diff --git a/Engine/Plugins/Experimental/NNERuntimeRDG/Source/NNEHlslShaders/Shaders/NNEHlslShadersConvMatmul.usf b/Engine/Plugins/Experimental/NNERuntimeRDG/Source/NNEHlslShaders/Shaders/NNEHlslShadersConvMatmul.usf
--- a/Engine/Plugins/Experimental/NNERuntimeRDG/Source/NNEHlslShaders/Shaders/NNEHlslShadersConvMatmul.usf
+++ b/Engine/Plugins/Experimental/NNERuntimeRDG/Source/NNEHlslShaders/Shaders/NNEHlslShadersConvMatmul.usf
@@ -71,7 +71,7 @@ void ConvMatmul(
 	const int Scalar_GroupOutputPixelOffset = TG_OP * GroupID.x;
 	const int Scalar_GroupOutputChannelOffset = TG_OC * GroupID.y;
 	const int Scalar_BatchOutputOffset = Ho * Wo * Cw * GroupID.z;
-	const int Scalar_BatchInputOffset = Hi * Wi * Cw * GroupID.z;
+	const int Scalar_BatchInputOffset = Hi * Wi * Ci * GroupID.z;
 	const int Scalar_GroupOutputPixelH = (Scalar_GroupOutputPixelOffset) / Wo * StrideH;
 	const int Scalar_GroupOutputPixelWBase = Scalar_GroupOutputPixelOffset % Wo * StrideW;
 	const int Scalar_GroupInputPixelKernelTopLeftH = Scalar_GroupOutputPixelH - PadTop;
