[NNE] NNERuntimeRDGHlsl Maxpool operator

Problem:
NNERuntimeRDGHlsl was missing Maxpool operator

Solution:
Implement it

Test:
Win64 Editor: all tests
Win64 Standalone: all tests

#rnx

diff --git a/Engine/Plugins/Experimental/NNERuntimeRDG/Source/NNEHlslShaders/Internal/NNEHlslShadersPoolCS.h b/Engine/Plugins/Experimental/NNERuntimeRDG/Source/NNEHlslShaders/Internal/NNEHlslShadersPoolCS.h
new file mode 100644
--- /dev/null
+++ b/Engine/Plugins/Experimental/NNERuntimeRDG/Source/NNEHlslShaders/Internal/NNEHlslShadersPoolCS.h
@@ -0,0 +1,41 @@
+// Copyright Epic Games, Inc. All Rights Reserved.
+
+#pragma once
+
+#include "NNEHlslShaderBase.h"
+#include "RenderGraphUtils.h"
+#include "ShaderParameterUtils.h"
+
+namespace UE::NNEHlslShaders::Internal
+{
+	class FPoolConstants
+	{
+	public:
+
+		static const int32 MAX_NUM_DIMENSIONS{ 8 };
+		static const int32 MAX_NUM_SPATIAL_DIMENSIONS{ 6 };
+		static const int32 NUM_GROUP_THREADS{ 256 };
+	};
+
+	class NNEHLSLSHADERS_API FPoolCS : public FHlslShaderBase
+	{
+		DECLARE_GLOBAL_SHADER(FPoolCS);
+		SHADER_USE_PARAMETER_STRUCT(FPoolCS, FHlslShaderBase)
+
+		class FPoolNumSpatialDimensions : SHADER_PERMUTATION_RANGE_INT("NUM_SPATIAL_DIMENSIONS", 1, FPoolConstants::MAX_NUM_SPATIAL_DIMENSIONS);
+		using FPermutationDomain = TShaderPermutationDomain<FPoolNumSpatialDimensions>;
+
+	public:
+
+		BEGIN_SHADER_PARAMETER_STRUCT(FParameters, )
+			SHADER_PARAMETER_RDG_BUFFER_SRV(Buffer<float>, Input)
+			SHADER_PARAMETER_RDG_BUFFER_UAV(RWBuffer<float>, Output)
+			SHADER_PARAMETER_ARRAY(FUintVector4, TensorInfo, [FPoolConstants::MAX_NUM_DIMENSIONS])
+			SHADER_PARAMETER_ARRAY(FUintVector4, SpatialInfo, [FPoolConstants::MAX_NUM_SPATIAL_DIMENSIONS])
+			SHADER_PARAMETER(uint32, Num)
+			SHADER_PARAMETER(uint32, ThreadCountX)
+		END_SHADER_PARAMETER_STRUCT()
+
+		static void ModifyCompilationEnvironment(const FGlobalShaderPermutationParameters& InParameters, FShaderCompilerEnvironment& OutEnvironment);
+	};
+} // UE::NNEHlslShaders::Internal
diff --git a/Engine/Plugins/Experimental/NNERuntimeRDG/Source/NNEHlslShaders/Private/NNEHlslShadersPoolCS.cpp b/Engine/Plugins/Experimental/NNERuntimeRDG/Source/NNEHlslShaders/Private/NNEHlslShadersPoolCS.cpp
new file mode 100644
--- /dev/null
+++ b/Engine/Plugins/Experimental/NNERuntimeRDG/Source/NNEHlslShaders/Private/NNEHlslShadersPoolCS.cpp
@@ -0,0 +1,18 @@
+// Copyright Epic Games, Inc. All Rights Reserved.
+
+#include "NNEHlslShadersPoolCS.h"
+#include "NNE.h"
+
+namespace UE::NNEHlslShaders::Internal
+{
+	void FPoolCS::ModifyCompilationEnvironment(const FGlobalShaderPermutationParameters& InParameters, FShaderCompilerEnvironment& OutEnvironment)
+	{
+		FGlobalShader::ModifyCompilationEnvironment(InParameters, OutEnvironment);
+
+		OutEnvironment.SetDefine(TEXT("THREADGROUP_SIZE_X"), FPoolConstants::NUM_GROUP_THREADS);
+
+		FPermutationDomain PermutationVector(InParameters.PermutationId);
+	}
+
+	IMPLEMENT_GLOBAL_SHADER(FPoolCS, "/NNE/NNEHlslShadersPool.usf", "Pool", SF_Compute);
+} // UE::NNEHlslShaders::Internal
\ No newline at end of file
diff --git a/Engine/Plugins/Experimental/NNERuntimeRDG/Source/NNEHlslShaders/Shaders/NNEHlslShadersPool.usf b/Engine/Plugins/Experimental/NNERuntimeRDG/Source/NNEHlslShaders/Shaders/NNEHlslShadersPool.usf
new file mode 100644
--- /dev/null
+++ b/Engine/Plugins/Experimental/NNERuntimeRDG/Source/NNEHlslShaders/Shaders/NNEHlslShadersPool.usf
@@ -0,0 +1,111 @@
+// Copyright Epic Games, Inc. All Rights Reserved.
+
+#include "/Engine/Public/Platform.ush"
+#include "NNEHlslShadersBroadcastHelper.ush"
+
+#define NUM_DIMENSION (NUM_SPATIAL_DIMENSIONS+2)
+
+Buffer<float> Input;
+RWBuffer<float> Output;
+uint4 TensorInfo[NUM_DIMENSION];
+uint4 SpatialInfo[NUM_SPATIAL_DIMENSIONS];
+uint Num;
+uint ThreadCountX;
+
+#define OUTPUT_STRIDE 0
+#define INPUT_STRIDE 1
+#define INPUT_SIZE 2
+
+#define SPATIAL_STRIDE 0
+#define SPATIAL_KERNEL 1
+#define SPATIAL_PAD_START 2
+
+[numthreads(THREADGROUP_SIZE_X, 1, 1)]
+void Pool(in const uint3 DispatchThreadID : SV_DispatchThreadID)
+{
+	const uint Index = DispatchThreadID.y * ThreadCountX + DispatchThreadID.x;
+	
+	if (Index < Num)
+	{
+		//Compute indices on output tensors
+		uint Offset = Index;
+		uint OutputIndices[NUM_DIMENSION];
+		for (uint dim = 0; dim < NUM_DIMENSION; ++dim)
+		{
+			uint OutDimIdx, R;
+			DivMod(Offset, TensorInfo[dim][OUTPUT_STRIDE], OutputIndices[dim], R);
+			Offset = R;
+		}
+
+		//Compute spatial region to apply pooling on input
+		uint SpatialDimStarts[NUM_SPATIAL_DIMENSIONS];
+		uint SpatialDimEnds[NUM_SPATIAL_DIMENSIONS];
+		for (uint spatialDim = 0; spatialDim < NUM_SPATIAL_DIMENSIONS; ++spatialDim)
+		{
+			const uint Dim = spatialDim+2;
+			const int SpatialDimStart = ((int)OutputIndices[Dim] * (int)SpatialInfo[spatialDim][SPATIAL_STRIDE]) - (int)SpatialInfo[spatialDim][SPATIAL_PAD_START];
+			const int SpatialDimEnd = SpatialDimStart + (int)SpatialInfo[spatialDim][SPATIAL_KERNEL];
+
+			SpatialDimStarts[spatialDim] = (uint)max(SpatialDimStart, (int)0);
+			SpatialDimEnds[spatialDim] = (uint)min(SpatialDimEnd, (int)TensorInfo[Dim][INPUT_SIZE]);
+		}
+
+		const uint BaseOffset = (TensorInfo[0][INPUT_STRIDE] * OutputIndices[0]) + (TensorInfo[1][INPUT_STRIDE] * OutputIndices[1]);
+		uint SpatialDimIterator[NUM_SPATIAL_DIMENSIONS];
+
+		//Apply pooling, looping other the input for all spatial dimensions of the input kernel
+		float MaxValue;
+		bool IsInit = true;
+		for (SpatialDimIterator[0] = SpatialDimStarts[0]; SpatialDimIterator[0] < SpatialDimEnds[0]; ++SpatialDimIterator[0])
+		{
+		#if (NUM_SPATIAL_DIMENSIONS > 1)
+		for (SpatialDimIterator[1] = SpatialDimStarts[1]; SpatialDimIterator[1] < SpatialDimEnds[1]; ++SpatialDimIterator[1])
+		{
+		#endif
+		#if (NUM_SPATIAL_DIMENSIONS > 2)
+		for (SpatialDimIterator[2] = SpatialDimStarts[2]; SpatialDimIterator[2] < SpatialDimEnds[2]; ++SpatialDimIterator[2])
+		{
+		#endif
+		#if (NUM_SPATIAL_DIMENSIONS > 3)
+		for (SpatialDimIterator[3] = SpatialDimStarts[3]; SpatialDimIterator[3] < SpatialDimEnds[3]; ++SpatialDimIterator[3])
+		{
+		#endif
+		#if (NUM_SPATIAL_DIMENSIONS > 4)
+		for (SpatialDimIterator[4] = SpatialDimStarts[4]; SpatialDimIterator[4] < SpatialDimEnds[4]; ++SpatialDimIterator[4])
+		{
+		#endif
+		#if (NUM_SPATIAL_DIMENSIONS > 5)
+		for (SpatialDimIterator[5] = SpatialDimStarts[1]; SpatialDimIterator[5] < SpatialDimEnds[5]; ++SpatialDimIterator[5])
+		{
+		#endif
+
+			uint InputIndex = BaseOffset;
+			[unroll]
+			for (uint spatialDim = 0; spatialDim < NUM_SPATIAL_DIMENSIONS; ++spatialDim)
+			{
+				InputIndex += TensorInfo[spatialDim+2][INPUT_STRIDE] * SpatialDimIterator[spatialDim];
+			}
+			float CurrentValue = Input[InputIndex];
+			MaxValue = IsInit ? CurrentValue : max(MaxValue, CurrentValue);
+			IsInit = false;
+
+		#if (NUM_SPATIAL_DIMENSIONS > 5)
+		}
+		#endif
+		#if (NUM_SPATIAL_DIMENSIONS > 4)
+		}
+		#endif
+		#if (NUM_SPATIAL_DIMENSIONS > 3)
+		}
+		#endif
+		#if (NUM_SPATIAL_DIMENSIONS > 2)
+		}
+		#endif
+		#if (NUM_SPATIAL_DIMENSIONS > 1)
+		}
+		#endif
+		}
+
+		Output[Index] = MaxValue;
+	}
+}
diff --git a/Engine/Plugins/Experimental/NNERuntimeRDG/Source/NNERuntimeRDG/Private/Hlsl/NNERuntimeRDGPool.cpp b/Engine/Plugins/Experimental/NNERuntimeRDG/Source/NNERuntimeRDG/Private/Hlsl/NNERuntimeRDGPool.cpp
new file mode 100644
--- /dev/null
+++ b/Engine/Plugins/Experimental/NNERuntimeRDG/Source/NNERuntimeRDG/Private/Hlsl/NNERuntimeRDGPool.cpp
@@ -0,0 +1,219 @@
+// Copyright Epic Games, Inc. All Rights Reserved.
+
+#include "NNERuntimeRDGPool.h"
+#include "NNEHlslShadersConvCS.h"
+#include "NNEHlslShadersPoolCS.h"
+#include "NNERuntimeRDGHlslHelper.h"
+#include "NNETensor.h"
+#include "NNETypes.h"
+#include "RenderGraphUtils.h"
+
+
+namespace UE::NNERuntimeRDG::Private::Hlsl
+{
+	DECLARE_GPU_STAT_NAMED(FNNEOperatorMaxPool, TEXT("NNE.Operator.Hlsl.MaxPool"));
+
+	/**
+	 * MaxPool operator implementation
+	 */
+	class FMaxPool : public FOperatorHlsl
+	{
+	public:
+
+		FMaxPool() {}
+		virtual ~FMaxPool() = default;
+
+		int32 NumSpatialDimensions = 0;
+		NNEHlslShaders::Internal::EConvAutoPad AutoPad = NNEHlslShaders::Internal::EConvAutoPad::NOTSET;
+		TArray<int32> Pads;
+		TArray<int32> Strides;
+		TArray<int32> KernelShape;
+
+	public:
+
+		virtual int PrepareOutputs(TConstArrayView<NNE::Internal::FTensorRef> InputTensors, TArrayView<NNE::Internal::FTensorRef> OutputTensors) const override
+		{
+			check(InputTensors.Num() == 1);
+			check(OutputTensors.Num() == 1);
+
+			const NNE::Internal::FTensor& X = *InputTensors[0];
+			TConstArrayView<uint32> InputShape = X.GetShape().GetData();
+			TArray<uint32, TInlineAllocator<NNE::FTensorShape::MaxRank>> OutputShape;
+
+			check(Pads.Num() == 2*NumSpatialDimensions);
+			check(Strides.Num() == NumSpatialDimensions);
+			check(KernelShape.Num() == NumSpatialDimensions);
+
+			OutputShape.SetNumUninitialized(InputShape.Num());
+			OutputShape[0] = InputShape[0];
+			OutputShape[1] = InputShape[1];
+			for (int32 i = 0; i < NumSpatialDimensions; ++i)
+			{
+				//See [REDACTED_PROFILE_URL]/onnx/blob/main/docs/Changelog.md#MaxPool-8
+				switch (AutoPad)
+				{
+					case NNEHlslShaders::Internal::EConvAutoPad::NOTSET:
+					{
+						uint32 PadShape = Pads[i] + Pads[i+NumSpatialDimensions];
+						OutputShape[i + 2] = (uint32)FMath::FloorToFloat((float)(InputShape[i + 2] + PadShape - KernelShape[i]) / (float)Strides[i] + 1.0f);
+						break;
+					}
+					case NNEHlslShaders::Internal::EConvAutoPad::VALID:
+					{
+						OutputShape[i + 2] = (uint32)FMath::CeilToFloat((float)(InputShape[i + 2] - KernelShape[i] + 1) / (float)Strides[i]);
+						break;
+					}
+					case NNEHlslShaders::Internal::EConvAutoPad::SAME_UPPER:
+					case NNEHlslShaders::Internal::EConvAutoPad::SAME_LOWER:
+					{
+						OutputShape[i + 2] = (uint32)FMath::CeilToFloat((float)InputShape[i + 2] / (float)Strides[i]);
+						break;
+					}
+				}
+			}
+			OutputTensors[0]->SetShape(NNE::FTensorShape::Make(OutputShape));
+
+			return 0;
+		};
+
+		virtual bool Initialize(TConstArrayView<NNE::FTensorDesc> InputTensorDescs, TConstArrayView<NNE::FTensorDesc> OutputTensorDescs, const NNE::FAttributeMap& Attributes) override
+		{
+			check(InputTensorDescs.Num() == 1);
+			check(OutputTensorDescs.Num() >= 1);
+
+			if (OutputTensorDescs.Num() > 1)
+			{
+				UE_LOG(LogNNE, Warning, TEXT("MaxPool 2nd optional output 'Indices' is not supported."));
+				return false;
+			}
+
+			const NNE::FTensorDesc& Input = InputTensorDescs[0];
+			const NNE::FTensorDesc& Output = OutputTensorDescs[0];
+
+			if (Input.GetShape().Rank() < 3)
+			{
+				UE_LOG(LogNNE, Warning, TEXT("MaxPool input should be at least of rank 3, to have 1+ spatial dimension(s) but is of rank %d"), Input.GetShape().Rank());
+				return false;
+			}
+			NumSpatialDimensions = Input.GetShape().Rank() - 2;
+
+			TArray<int32> StridesDefault;
+			TArray<int32> PadsDefault;
+			StridesDefault.Init(1, NumSpatialDimensions);
+			PadsDefault.Init(0, 2 * NumSpatialDimensions);
+
+			NNEHlslShaders::Internal::FConvCS::LexFromString(AutoPad, *Attributes.GetValueOrDefault<FString>(TEXT("auto_pad"), TEXT("NOTSET")));
+			Pads = Attributes.GetValueOrDefault<TArray<int32>>(TEXT("pads"), PadsDefault);
+			Strides = Attributes.GetValueOrDefault<TArray<int32>>(TEXT("strides"), StridesDefault);
+			KernelShape = Attributes.GetValue<TArray<int32>>(TEXT("kernel_shape"));
+
+			if (KernelShape.Num() != NumSpatialDimensions)
+			{
+				UE_LOG(LogNNE, Warning, TEXT("MaxPool KernelShape should have as many elements as the spatial dimensions of the input, got %d while input have %d."), KernelShape.Num(), NumSpatialDimensions);
+				return false;
+			}
+			if (Strides.Num() != NumSpatialDimensions)
+			{
+				UE_LOG(LogNNE, Warning, TEXT("MaxPool Strides should have as many elements as the spatial dimensions of the input, got %d while input have %d."), Strides.Num(), NumSpatialDimensions);
+				return false;
+			}
+			if (Pads.Num() != 2*NumSpatialDimensions)
+			{
+				UE_LOG(LogNNE, Warning, TEXT("MaxPool Pads should have twice as many elements as the spatial dimensions of the input, got %d while input have %d."), Pads.Num(), NumSpatialDimensions);
+				return false;
+			}
+
+			return true;
+		}
+
+		virtual void Dispatch(FRDGBuilder& GraphBuilder, TConstArrayView<FTensorRDGRef> InputTensors, TConstArrayView<FTensorRDGRef> OutputTensors) override
+		{
+			using namespace UE::NNEHlslShaders::Internal;
+
+			check(InputTensors.Num() == 1);
+			check(OutputTensors.Num() == 1);
+			check(InputTensors[0] != nullptr);
+			check(OutputTensors[0] != nullptr);
+
+			const FTensorRDG& Input = *InputTensors[0];
+			const FTensorRDG& Output = *OutputTensors[0];
+			const FRDGBufferSRVRef InputSRV = GraphBuilder.CreateSRV(FRDGBufferSRVDesc(Input.GetBuffer(), PF_R32_FLOAT));
+			const FRDGBufferUAVRef OutputUAV = GraphBuilder.CreateUAV(FRDGBufferUAVDesc(Output.GetBuffer(), PF_R32_FLOAT));
+			const FIntVector ThreadGroupCount = ComputeElementWiseThreadGroups(Output.GetVolume(), FPoolConstants::NUM_GROUP_THREADS);
+
+			// Set parameters
+			FPoolCS::FParameters* Params = GraphBuilder.AllocParameters<FPoolCS::FParameters>();
+			Params->Input = InputSRV;
+			Params->Output = OutputUAV;
+			Params->Num = Output.GetVolume();
+			Params->ThreadCountX = ThreadGroupCount.X * FPoolConstants::NUM_GROUP_THREADS;
+			FillTensorStrideShaderParameters(Output, Params->TensorInfo, 0);
+			FillTensorStrideShaderParameters(Input, Params->TensorInfo, 1);
+			FillTensorSizeShaderParameters(Input, Params->TensorInfo, 2);
+			for (int32 i = 0; i < NumSpatialDimensions; ++i)
+			{
+				Params->SpatialInfo[i][0] = Strides[i];
+				Params->SpatialInfo[i][1] = KernelShape[i];
+				if (AutoPad == NNEHlslShaders::Internal::EConvAutoPad::SAME_LOWER)
+				{
+					// see [REDACTED_PROFILE_URL]/onnx/blob/main/docs/Changelog.md#MaxPool-8
+					// only needed for SAME_LOWER as SAME_UPPER is handled by index clamping in the HLSL kernel
+					Params->SpatialInfo[i][2] = (Output.GetShape().GetData()[i + 2] - 1) * Strides[i] + KernelShape[i] - Input.GetShape().GetData()[i + 2];
+				}
+				else
+				{
+					Params->SpatialInfo[i][2] = Pads[i];
+				}
+			}
+
+			FPoolCS::FPermutationDomain PermutationVector;
+			PermutationVector.Set<FPoolCS::FPoolNumSpatialDimensions>(NumSpatialDimensions);
+
+			TShaderMapRef<FPoolCS> ComputeShader(GetGlobalShaderMap(GMaxRHIFeatureLevel), PermutationVector);
+
+			RDG_EVENT_SCOPE(GraphBuilder, "NNE.Operator.Hlsl.MaxPool");
+			RDG_GPU_STAT_SCOPE(GraphBuilder, FNNEOperatorMaxPool);
+
+			FComputeShaderUtils::AddPass(
+				GraphBuilder,
+				RDG_EVENT_NAME("NNE.Operator.Hlsl.MaxPool.Dispatch"),
+				ERDGPassFlags::Compute | ERDGPassFlags::NeverCull,
+				ComputeShader,
+				Params,
+				ThreadGroupCount);
+		}
+	};
+
+	bool ValidateMaxPoolOperator(const NNE::FAttributeMap& AttributeMap, TConstArrayView<ENNETensorDataType> InputTypes, TConstArrayView<NNE::FSymbolicTensorShape> InputPools)
+	{
+		//This match version 8 of the MaxPool operator, next version is 10
+		//[REDACTED_PROFILE_URL]/onnx/blob/main/docs/Operators.md#MaxPool
+		bool bIsValid = true;
+
+		FAttributeValidator AttributeValidator;
+		AttributeValidator.AddOptional(TEXT("auto_pad"), ENNEAttributeDataType::String);
+		AttributeValidator.AddRequired(TEXT("kernel_shape"), ENNEAttributeDataType::Int32Array);
+		AttributeValidator.AddOptional(TEXT("pads"), ENNEAttributeDataType::Int32Array);
+		AttributeValidator.AddOptional(TEXT("storage_order"), ENNEAttributeDataType::Int32);//Unused, only needed for 2nd output itself not supported, see [REDACTED_PROFILE_URL]/onnx/issues/1370
+		AttributeValidator.AddOptional(TEXT("strides"), ENNEAttributeDataType::Int32Array);
+		bIsValid &= AttributeValidator.Validate(AttributeMap);
+
+		FInputValidator InputValidator;
+		InputValidator.AddSupportedType(ENNETensorDataType::Float);
+		InputValidator.AddRequired();
+		bIsValid &= InputValidator.Validate(InputTypes);
+
+		return bIsValid;
+	}
+
+	FOperatorHlsl* CreateMaxPoolOperator()
+	{
+		return new FMaxPool();
+	}
+
+	bool RegisterPoolOperators(FOperatorRegistryHlsl& Registry)
+	{
+		Registry.OpAdd({{TEXT("MaxPool"), TEXT("Onnx")}}, CreateMaxPoolOperator, ValidateMaxPoolOperator);
+		return true;
+	}
+} // UE::NNERuntimeRDG::Private::Hlsl
diff --git a/Engine/Plugins/Experimental/NNERuntimeRDG/Source/NNERuntimeRDG/Private/Hlsl/NNERuntimeRDGPool.h b/Engine/Plugins/Experimental/NNERuntimeRDG/Source/NNERuntimeRDG/Private/Hlsl/NNERuntimeRDGPool.h
new file mode 100644
--- /dev/null
+++ b/Engine/Plugins/Experimental/NNERuntimeRDG/Source/NNERuntimeRDG/Private/Hlsl/NNERuntimeRDGPool.h
@@ -0,0 +1,11 @@
+// Copyright Epic Games, Inc. All Rights Reserved.
+
+#pragma once
+
+#include "NNERuntimeRDGBase.h"
+#include "NNERuntimeRDGHlslOp.h"
+
+namespace UE::NNERuntimeRDG::Private::Hlsl
+{
+	bool RegisterPoolOperators(FOperatorRegistryHlsl& Registry);
+} // UE::NNERuntimeRDG::Private::Hlsl
diff --git a/Engine/Plugins/Experimental/NNERuntimeRDG/Source/NNERuntimeRDG/Private/NNERuntimeRDGHlsl.cpp b/Engine/Plugins/Experimental/NNERuntimeRDG/Source/NNERuntimeRDG/Private/NNERuntimeRDGHlsl.cpp
--- a/Engine/Plugins/Experimental/NNERuntimeRDG/Source/NNERuntimeRDG/Private/NNERuntimeRDGHlsl.cpp
+++ b/Engine/Plugins/Experimental/NNERuntimeRDG/Source/NNERuntimeRDG/Private/NNERuntimeRDGHlsl.cpp
@@ -29,6 +29,7 @@
 #include "Hlsl/NNERuntimeRDGIdentity.h"
 #include "Hlsl/NNERuntimeRDGInstanceNormalization.h"
 #include "Hlsl/NNERuntimeRDGPad.h"
+#include "Hlsl/NNERuntimeRDGPool.h"
 #include "Hlsl/NNERuntimeRDGReshape.h"
 #include "Hlsl/NNERuntimeRDGShape.h"
 #include "Hlsl/NNERuntimeRDGSize.h"
@@ -64,6 +65,7 @@ bool UNNERuntimeRDGHlslImpl::Init()
 	RegisterIdentityOperator(*registry);
 	RegisterInstanceNormalizationOperator(*registry);
 	RegisterPadOperator(*registry);
+	RegisterPoolOperators(*registry);
 	RegisterReshapeOperator(*registry);
 	RegisterShapeOperator(*registry);
 	RegisterSizeOperator(*registry);
