[NNE] NNERuntimeRDGHlsl AveragePool operator

Problem:
NNERuntimeRDGHlsl was missing AveragePool operator

Solution:
Implement it

Test:
Win64 Editor: all tests
Win64 Standalone: all tests

#rnx

diff --git a/Engine/Plugins/Experimental/NNERuntimeRDG/Source/NNEHlslShaders/Internal/NNEHlslShadersPoolCS.h b/Engine/Plugins/Experimental/NNERuntimeRDG/Source/NNEHlslShaders/Internal/NNEHlslShadersPoolCS.h
--- a/Engine/Plugins/Experimental/NNERuntimeRDG/Source/NNEHlslShaders/Internal/NNEHlslShadersPoolCS.h
+++ b/Engine/Plugins/Experimental/NNERuntimeRDG/Source/NNEHlslShaders/Internal/NNEHlslShadersPoolCS.h
@@ -8,6 +8,13 @@
 
 namespace UE::NNEHlslShaders::Internal
 {
+	enum class EPoolOperatorType : uint8
+	{
+		MAX_POOL,
+		AVERAGE_POOL,
+		MAX
+	};
+
 	class FPoolConstants
 	{
 	public:
@@ -23,7 +30,8 @@ namespace UE::NNEHlslShaders::Internal
 		SHADER_USE_PARAMETER_STRUCT(FPoolCS, FHlslShaderBase)
 
 		class FPoolNumSpatialDimensions : SHADER_PERMUTATION_RANGE_INT("NUM_SPATIAL_DIMENSIONS", 1, FPoolConstants::MAX_NUM_SPATIAL_DIMENSIONS);
-		using FPermutationDomain = TShaderPermutationDomain<FPoolNumSpatialDimensions>;
+		class FPoolType : SHADER_PERMUTATION_ENUM_CLASS("POOL_OPERATOR_TYPE", EPoolOperatorType);
+		using FPermutationDomain = TShaderPermutationDomain<FPoolNumSpatialDimensions, FPoolType>;
 
 	public:
 
@@ -34,6 +42,7 @@ namespace UE::NNEHlslShaders::Internal
 			SHADER_PARAMETER_ARRAY(FUintVector4, SpatialInfo, [FPoolConstants::MAX_NUM_SPATIAL_DIMENSIONS])
 			SHADER_PARAMETER(uint32, Num)
 			SHADER_PARAMETER(uint32, ThreadCountX)
+			SHADER_PARAMETER(uint32, KernelVolume)
 		END_SHADER_PARAMETER_STRUCT()
 
 		static void ModifyCompilationEnvironment(const FGlobalShaderPermutationParameters& InParameters, FShaderCompilerEnvironment& OutEnvironment);
diff --git a/Engine/Plugins/Experimental/NNERuntimeRDG/Source/NNEHlslShaders/Shaders/NNEHlslShadersPool.usf b/Engine/Plugins/Experimental/NNERuntimeRDG/Source/NNEHlslShaders/Shaders/NNEHlslShadersPool.usf
--- a/Engine/Plugins/Experimental/NNERuntimeRDG/Source/NNEHlslShaders/Shaders/NNEHlslShadersPool.usf
+++ b/Engine/Plugins/Experimental/NNERuntimeRDG/Source/NNEHlslShaders/Shaders/NNEHlslShadersPool.usf
@@ -11,6 +11,7 @@ uint4 TensorInfo[NUM_DIMENSION];
 uint4 SpatialInfo[NUM_SPATIAL_DIMENSIONS];
 uint Num;
 uint ThreadCountX;
+uint KernelVolume;
 
 #define OUTPUT_STRIDE 0
 #define INPUT_STRIDE 1
@@ -20,6 +21,9 @@ uint ThreadCountX;
 #define SPATIAL_KERNEL 1
 #define SPATIAL_PAD_START 2
 
+#define MAX_POOL 0
+#define AVERAGE_POOL 1
+
 [numthreads(THREADGROUP_SIZE_X, 1, 1)]
 void Pool(in const uint3 DispatchThreadID : SV_DispatchThreadID)
 {
@@ -54,8 +58,9 @@ void Pool(in const uint3 DispatchThreadID : SV_DispatchThreadID)
 		uint SpatialDimIterator[NUM_SPATIAL_DIMENSIONS];
 
 		//Apply pooling, looping other the input for all spatial dimensions of the input kernel
-		float MaxValue;
-		bool IsInit = true;
+		float PooledValue = 0;
+		int Count = 0;
+		
 		for (SpatialDimIterator[0] = SpatialDimStarts[0]; SpatialDimIterator[0] < SpatialDimEnds[0]; ++SpatialDimIterator[0])
 		{
 		#if (NUM_SPATIAL_DIMENSIONS > 1)
@@ -86,8 +91,12 @@ void Pool(in const uint3 DispatchThreadID : SV_DispatchThreadID)
 				InputIndex += TensorInfo[spatialDim+2][INPUT_STRIDE] * SpatialDimIterator[spatialDim];
 			}
 			float CurrentValue = Input[InputIndex];
-			MaxValue = IsInit ? CurrentValue : max(MaxValue, CurrentValue);
-			IsInit = false;
+			#if (POOL_OPERATOR_TYPE == MAX_POOL)
+				PooledValue = (Count == 0) ? CurrentValue : max(PooledValue, CurrentValue);
+			#else //POOL_OPERATOR_TYPE == AVERAGE_POOL
+				PooledValue += CurrentValue;
+			#endif 
+			++Count;
 
 		#if (NUM_SPATIAL_DIMENSIONS > 5)
 		}
@@ -106,6 +115,10 @@ void Pool(in const uint3 DispatchThreadID : SV_DispatchThreadID)
 		#endif
 		}
 
-		Output[Index] = MaxValue;
+		#if (POOL_OPERATOR_TYPE == AVERAGE_POOL)
+			PooledValue /= (KernelVolume == 0) ? (float)Count : (float)KernelVolume;
+		#endif
+
+		Output[Index] = PooledValue;
 	}
 }
diff --git a/Engine/Plugins/Experimental/NNERuntimeRDG/Source/NNERuntimeRDG/Private/Hlsl/NNERuntimeRDGPool.cpp b/Engine/Plugins/Experimental/NNERuntimeRDG/Source/NNERuntimeRDG/Private/Hlsl/NNERuntimeRDGPool.cpp
--- a/Engine/Plugins/Experimental/NNERuntimeRDG/Source/NNERuntimeRDG/Private/Hlsl/NNERuntimeRDGPool.cpp
+++ b/Engine/Plugins/Experimental/NNERuntimeRDG/Source/NNERuntimeRDG/Private/Hlsl/NNERuntimeRDGPool.cpp
@@ -12,22 +12,25 @@
 namespace UE::NNERuntimeRDG::Private::Hlsl
 {
 	DECLARE_GPU_STAT_NAMED(FNNEOperatorMaxPool, TEXT("NNE.Operator.Hlsl.MaxPool"));
+	DECLARE_GPU_STAT_NAMED(FNNEOperatorAveragePool, TEXT("NNE.Operator.Hlsl.AveragePool"));
 
 	/**
 	 * MaxPool operator implementation
 	 */
-	class FMaxPool : public FOperatorHlsl
+	template< UE::NNEHlslShaders::Internal::EPoolOperatorType PoolOperatorType >
+	class FPoolOperator : public FOperatorHlsl
 	{
 	public:
 
-		FMaxPool() {}
-		virtual ~FMaxPool() = default;
+		FPoolOperator() {}
+		virtual ~FPoolOperator() = default;
 
 		int32 NumSpatialDimensions = 0;
 		NNEHlslShaders::Internal::EConvAutoPad AutoPad = NNEHlslShaders::Internal::EConvAutoPad::NOTSET;
 		TArray<int32> Pads;
 		TArray<int32> Strides;
 		TArray<int32> KernelShape;
+		int32 KernelVolume = 0;
 
 	public:
 
@@ -79,12 +82,19 @@ namespace UE::NNERuntimeRDG::Private::Hlsl
 		virtual bool Initialize(TConstArrayView<NNE::FTensorDesc> InputTensorDescs, TConstArrayView<NNE::FTensorDesc> OutputTensorDescs, const NNE::FAttributeMap& Attributes) override
 		{
 			check(InputTensorDescs.Num() == 1);
-			check(OutputTensorDescs.Num() >= 1);
 
-			if (OutputTensorDescs.Num() > 1)
+			if constexpr (PoolOperatorType == UE::NNEHlslShaders::Internal::EPoolOperatorType::MAX_POOL)
 			{
-				UE_LOG(LogNNE, Warning, TEXT("MaxPool 2nd optional output 'Indices' is not supported."));
-				return false;
+				check(OutputTensorDescs.Num() >= 1);
+				if (OutputTensorDescs.Num() > 1)
+				{
+					UE_LOG(LogNNE, Warning, TEXT("MaxPool 2nd optional output 'Indices' is not supported."));
+					return false;
+				}
+			}
+			else
+			{
+				check(OutputTensorDescs.Num() == 1);
 			}
 
 			const NNE::FTensorDesc& Input = InputTensorDescs[0];
@@ -92,7 +102,7 @@ namespace UE::NNERuntimeRDG::Private::Hlsl
 
 			if (Input.GetShape().Rank() < 3)
 			{
-				UE_LOG(LogNNE, Warning, TEXT("MaxPool input should be at least of rank 3, to have 1+ spatial dimension(s) but is of rank %d"), Input.GetShape().Rank());
+				UE_LOG(LogNNE, Warning, TEXT("%s input should be at least of rank 3, to have 1+ spatial dimension(s) but is of rank %d"), GetOperatorName(), Input.GetShape().Rank());
 				return false;
 			}
 			NumSpatialDimensions = Input.GetShape().Rank() - 2;
@@ -109,20 +119,34 @@ namespace UE::NNERuntimeRDG::Private::Hlsl
 
 			if (KernelShape.Num() != NumSpatialDimensions)
 			{
-				UE_LOG(LogNNE, Warning, TEXT("MaxPool KernelShape should have as many elements as the spatial dimensions of the input, got %d while input have %d."), KernelShape.Num(), NumSpatialDimensions);
+				UE_LOG(LogNNE, Warning, TEXT("%s KernelShape should have as many elements as the spatial dimensions of the input, got %d while input have %d."), GetOperatorName(), KernelShape.Num(), NumSpatialDimensions);
 				return false;
 			}
 			if (Strides.Num() != NumSpatialDimensions)
 			{
-				UE_LOG(LogNNE, Warning, TEXT("MaxPool Strides should have as many elements as the spatial dimensions of the input, got %d while input have %d."), Strides.Num(), NumSpatialDimensions);
+				UE_LOG(LogNNE, Warning, TEXT("%s Strides should have as many elements as the spatial dimensions of the input, got %d while input have %d."), GetOperatorName(), Strides.Num(), NumSpatialDimensions);
 				return false;
 			}
 			if (Pads.Num() != 2*NumSpatialDimensions)
 			{
-				UE_LOG(LogNNE, Warning, TEXT("MaxPool Pads should have twice as many elements as the spatial dimensions of the input, got %d while input have %d."), Pads.Num(), NumSpatialDimensions);
+				UE_LOG(LogNNE, Warning, TEXT("%s Pads should have twice as many elements as the spatial dimensions of the input, got %d while input have %d."), GetOperatorName(), Pads.Num(), NumSpatialDimensions);
 				return false;
 			}
 
+			KernelVolume = 0;//If KernelVolume is 0 the kernel will use the count of pooled elements.
+			if constexpr (PoolOperatorType == UE::NNEHlslShaders::Internal::EPoolOperatorType::AVERAGE_POOL)
+			{
+				int32 CountIncludePad = Attributes.GetValueOrDefault<int32>(TEXT("count_include_pad"), 0);
+				if (CountIncludePad != 0)
+				{
+					KernelVolume = 1;
+					for (int32 dim : KernelShape)
+					{
+						KernelVolume *= dim;
+					}
+				}
+			}
+
 			return true;
 		}
 
@@ -147,6 +171,7 @@ namespace UE::NNERuntimeRDG::Private::Hlsl
 			Params->Output = OutputUAV;
 			Params->Num = Output.GetVolume();
 			Params->ThreadCountX = ThreadGroupCount.X * FPoolConstants::NUM_GROUP_THREADS;
+			Params->KernelVolume = KernelVolume;
 			FillTensorStrideShaderParameters(Output, Params->TensorInfo, 0);
 			FillTensorStrideShaderParameters(Input, Params->TensorInfo, 1);
 			FillTensorSizeShaderParameters(Input, Params->TensorInfo, 2);
@@ -168,25 +193,56 @@ namespace UE::NNERuntimeRDG::Private::Hlsl
 
 			FPoolCS::FPermutationDomain PermutationVector;
 			PermutationVector.Set<FPoolCS::FPoolNumSpatialDimensions>(NumSpatialDimensions);
+			PermutationVector.Set<FPoolCS::FPoolType>(PoolOperatorType);
 
 			TShaderMapRef<FPoolCS> ComputeShader(GetGlobalShaderMap(GMaxRHIFeatureLevel), PermutationVector);
-
-			RDG_EVENT_SCOPE(GraphBuilder, "NNE.Operator.Hlsl.MaxPool");
-			RDG_GPU_STAT_SCOPE(GraphBuilder, FNNEOperatorMaxPool);
-
-			FComputeShaderUtils::AddPass(
-				GraphBuilder,
-				RDG_EVENT_NAME("NNE.Operator.Hlsl.MaxPool.Dispatch"),
-				ERDGPassFlags::Compute | ERDGPassFlags::NeverCull,
-				ComputeShader,
-				Params,
-				ThreadGroupCount);
+			
+			if constexpr (PoolOperatorType == EPoolOperatorType::MAX_POOL)
+			{
+				RDG_EVENT_SCOPE(GraphBuilder, "NNE.Operator.Hlsl.MaxPool");
+				RDG_GPU_STAT_SCOPE(GraphBuilder, FNNEOperatorMaxPool);
+
+				FComputeShaderUtils::AddPass(
+					GraphBuilder,
+					RDG_EVENT_NAME("NNE.Operator.Hlsl.MaxPool.Dispatch"),
+					ERDGPassFlags::Compute | ERDGPassFlags::NeverCull,
+					ComputeShader,
+					Params,
+					ThreadGroupCount);
+			}
+			else
+			{
+				RDG_EVENT_SCOPE(GraphBuilder, "NNE.Operator.Hlsl.AveragePool");
+				RDG_GPU_STAT_SCOPE(GraphBuilder, FNNEOperatorAveragePool);
+
+				FComputeShaderUtils::AddPass(
+					GraphBuilder,
+					RDG_EVENT_NAME("NNE.Operator.Hlsl.AveragePool.Dispatch"),
+					ERDGPassFlags::Compute | ERDGPassFlags::NeverCull,
+					ComputeShader,
+					Params,
+					ThreadGroupCount);
+			}
+		}
+	private:
+		constexpr static auto GetOperatorName()
+		{
+			if constexpr (PoolOperatorType == UE::NNEHlslShaders::Internal::EPoolOperatorType::MAX_POOL)
+			{
+				return TEXT("MaxPool");
+			}
+			else
+			{
+				return TEXT("AveragePool");
+			}
 		}
 	};
 
-	bool ValidateMaxPoolOperator(const NNE::FAttributeMap& AttributeMap, TConstArrayView<ENNETensorDataType> InputTypes, TConstArrayView<NNE::FSymbolicTensorShape> InputPools)
+	template< UE::NNEHlslShaders::Internal::EPoolOperatorType PoolOperatorType >
+	bool ValidatePoolOperator(const NNE::FAttributeMap& AttributeMap, TConstArrayView<ENNETensorDataType> InputTypes, TConstArrayView<NNE::FSymbolicTensorShape> InputPools)
 	{
 		//This match version 8 of the MaxPool operator, next version is 10
+		//and version 7 of the AveragexPool operator, next version is 10
 		//[REDACTED_PROFILE_URL]/onnx/blob/main/docs/Operators.md#MaxPool
 		bool bIsValid = true;
 
@@ -194,7 +250,15 @@ namespace UE::NNERuntimeRDG::Private::Hlsl
 		AttributeValidator.AddOptional(TEXT("auto_pad"), ENNEAttributeDataType::String);
 		AttributeValidator.AddRequired(TEXT("kernel_shape"), ENNEAttributeDataType::Int32Array);
 		AttributeValidator.AddOptional(TEXT("pads"), ENNEAttributeDataType::Int32Array);
-		AttributeValidator.AddOptional(TEXT("storage_order"), ENNEAttributeDataType::Int32);//Unused, only needed for 2nd output itself not supported, see [REDACTED_PROFILE_URL]/onnx/issues/1370
+		if constexpr (PoolOperatorType == UE::NNEHlslShaders::Internal::EPoolOperatorType::MAX_POOL)
+		{
+			AttributeValidator.AddOptional(TEXT("storage_order"), ENNEAttributeDataType::Int32);//Unused, only needed for 2nd output itself not supported, see [REDACTED_PROFILE_URL]/onnx/issues/1370
+		}
+		else
+		{
+			AttributeValidator.AddOptional(TEXT("count_include_pad"), ENNEAttributeDataType::Int32);
+		}
+
 		AttributeValidator.AddOptional(TEXT("strides"), ENNEAttributeDataType::Int32Array);
 		bIsValid &= AttributeValidator.Validate(AttributeMap);
 
@@ -208,12 +272,18 @@ namespace UE::NNERuntimeRDG::Private::Hlsl
 
 	FOperatorHlsl* CreateMaxPoolOperator()
 	{
-		return new FMaxPool();
+		return new FPoolOperator<UE::NNEHlslShaders::Internal::EPoolOperatorType::MAX_POOL>();
+	}
+
+	FOperatorHlsl* CreateAveragePoolOperator()
+	{
+		return new FPoolOperator<UE::NNEHlslShaders::Internal::EPoolOperatorType::AVERAGE_POOL>();
 	}
 
 	bool RegisterPoolOperators(FOperatorRegistryHlsl& Registry)
 	{
-		Registry.OpAdd({{TEXT("MaxPool"), TEXT("Onnx")}}, CreateMaxPoolOperator, ValidateMaxPoolOperator);
+		Registry.OpAdd({{TEXT("MaxPool"), TEXT("Onnx")}}, CreateMaxPoolOperator, ValidatePoolOperator<UE::NNEHlslShaders::Internal::EPoolOperatorType::MAX_POOL>);
+		Registry.OpAdd({ {TEXT("AveragePool"), TEXT("Onnx")} }, CreateAveragePoolOperator, ValidatePoolOperator<UE::NNEHlslShaders::Internal::EPoolOperatorType::AVERAGE_POOL>);
 		return true;
 	}
 } // UE::NNERuntimeRDG::Private::Hlsl
