[NNE] Softmax operator implementation for NNERuntimeRDGHlsl.

Problem:
NNERuntimeRDGHlsl was missing Softmax operator

Solution:
Implement it.
From the work of [REDACTED_HANDLE] in CL 29885045.

Next step is to implement the ReduceSum op and factorize the reduction, for softmax/instNorm and reduce

Test:
Win64 Editor: smoke tests + softmax tests
Win64 Standalone: all tests

#rnx

#virtualized

diff --git a/Engine/Plugins/Experimental/NNERuntimeRDG/Source/NNEHlslShaders/Internal/NNEHlslShadersSoftmaxCS.h b/Engine/Plugins/Experimental/NNERuntimeRDG/Source/NNEHlslShaders/Internal/NNEHlslShadersSoftmaxCS.h
new file mode 100644
--- /dev/null
+++ b/Engine/Plugins/Experimental/NNERuntimeRDG/Source/NNEHlslShaders/Internal/NNEHlslShadersSoftmaxCS.h
@@ -0,0 +1,36 @@
+// Copyright Epic Games, Inc. All Rights Reserved.
+
+#pragma once
+
+#include "NNEHlslShaderBase.h"
+#include "RenderGraphUtils.h"
+#include "ShaderParameterUtils.h"
+
+namespace UE::NNE::Internal { class FTensor; }
+
+namespace UE::NNEHlslShaders::Internal
+{
+	class FSoftmaxConstants
+	{
+	public:
+		static const int32 NUM_GROUP_THREADS{ 768 };
+	};
+
+	class NNEHLSLSHADERS_API TSoftmaxCS : public FHlslShaderBase
+	{
+		DECLARE_GLOBAL_SHADER(TSoftmaxCS);
+		SHADER_USE_PARAMETER_STRUCT(TSoftmaxCS, FHlslShaderBase)
+
+	public:
+
+		BEGIN_SHADER_PARAMETER_STRUCT(FParameters, )
+			SHADER_PARAMETER(int32, N)
+			SHADER_PARAMETER(int32, D)
+			SHADER_PARAMETER_RDG_BUFFER_SRV(Buffer<float>, Input)
+			SHADER_PARAMETER_RDG_BUFFER_UAV(RWBuffer<float>, Output)
+		END_SHADER_PARAMETER_STRUCT()
+
+		static FIntVector GetGroupCount(const FParameters& Parameters);
+		static void ModifyCompilationEnvironment(const FGlobalShaderPermutationParameters& InParameters, FShaderCompilerEnvironment& OutEnvironment);
+	};
+} // UE::NNEHlslShaders::Internal
\ No newline at end of file
diff --git a/Engine/Plugins/Experimental/NNERuntimeRDG/Source/NNEHlslShaders/Private/NNEHlslShadersSoftmaxCS.cpp b/Engine/Plugins/Experimental/NNERuntimeRDG/Source/NNEHlslShaders/Private/NNEHlslShadersSoftmaxCS.cpp
new file mode 100644
--- /dev/null
+++ b/Engine/Plugins/Experimental/NNERuntimeRDG/Source/NNEHlslShaders/Private/NNEHlslShadersSoftmaxCS.cpp
@@ -0,0 +1,21 @@
+// Copyright Epic Games, Inc. All Rights Reserved.
+
+#include "NNEHlslShadersSoftmaxCS.h"
+#include "NNE.h"
+
+namespace UE::NNEHlslShaders::Internal
+{
+	FIntVector TSoftmaxCS::GetGroupCount(const TSoftmaxCS::FParameters& Parameters)
+	{
+		return { 1, (int32)Parameters.N, 1 };
+	}
+
+	void TSoftmaxCS::ModifyCompilationEnvironment(const FGlobalShaderPermutationParameters& InParameters, FShaderCompilerEnvironment& OutEnvironment)
+	{
+		FGlobalShader::ModifyCompilationEnvironment(InParameters, OutEnvironment);
+
+		OutEnvironment.SetDefine(TEXT("THREADGROUP_SIZE_X"), FSoftmaxConstants::NUM_GROUP_THREADS);
+	}
+
+	IMPLEMENT_GLOBAL_SHADER(TSoftmaxCS, "/NNE/NNEHlslShadersSoftmax.usf", "Softmax", SF_Compute);
+} // UE::NNEHlslShaders::Internal
diff --git a/Engine/Plugins/Experimental/NNERuntimeRDG/Source/NNEHlslShaders/Shaders/NNEHlslShadersSoftmax.usf b/Engine/Plugins/Experimental/NNERuntimeRDG/Source/NNEHlslShaders/Shaders/NNEHlslShadersSoftmax.usf
new file mode 100644
--- /dev/null
+++ b/Engine/Plugins/Experimental/NNERuntimeRDG/Source/NNEHlslShaders/Shaders/NNEHlslShadersSoftmax.usf
@@ -0,0 +1,72 @@
+// Copyright Epic Games, Inc. All Rights Reserved.
+
+#include "/Engine/Public/Platform.ush"
+
+#define WORK_TYPE float
+#define BUFFER_TYPE float
+#define READ(x) x
+#define WRITE(x) x
+
+Buffer<WORK_TYPE> Input;
+RWBuffer<WORK_TYPE> Output;
+int N;
+int D;
+
+groupshared WORK_TYPE SharedMemorySumExp[THREADGROUP_SIZE_X];
+
+[numthreads(THREADGROUP_SIZE_X, 1, 1)]
+void Softmax(in const uint3 DispatchThreadID : SV_DispatchThreadID, in const uint3 GroupThreadID : SV_GroupThreadID)
+{
+	const int StartIdx = DispatchThreadID.y * D;
+	const int SharedMemoryOffsetY = GroupThreadID.y * THREADGROUP_SIZE_X;
+
+	// Reduction using shared mem and THREADGROUP_SIZE_X thread
+	// https://en.wikipedia.org/wiki/Algorithms_for_calculating_variance
+	if (DispatchThreadID.y < N)
+	{
+		WORK_TYPE SumExp = 0;
+		int Count = 0;
+
+		for (int CurrIdx = GroupThreadID.x; CurrIdx < D; CurrIdx += THREADGROUP_SIZE_X)
+		{
+			Count += 1;
+
+			WORK_TYPE Value = READ(Input[StartIdx + CurrIdx]);
+
+			SumExp += exp(Value);
+		}
+
+		SharedMemorySumExp[SharedMemoryOffsetY + GroupThreadID.x] = SumExp;
+	}
+
+	// Wait until all threads wrote result to shared memory
+	GroupMemoryBarrierWithGroupSync();
+
+	// Use one thread to calculate final result
+	if (GroupThreadID.x == 0 && DispatchThreadID.y < N)
+	{
+		WORK_TYPE SumExp = 0;
+
+		UNROLL
+		for (int Idx = 0; Idx < THREADGROUP_SIZE_X; Idx++)
+		{
+			WORK_TYPE CurrSumExp = SharedMemorySumExp[SharedMemoryOffsetY + Idx];
+			SumExp += CurrSumExp;
+		}
+
+		SharedMemorySumExp[SharedMemoryOffsetY] = SumExp;
+	}
+
+	// Wait until first thread wrote result to shared memory
+	GroupMemoryBarrierWithGroupSync();
+
+	if (DispatchThreadID.y < N)
+	{
+		WORK_TYPE SumExp = SharedMemorySumExp[SharedMemoryOffsetY];
+
+		for (int CurrIdx = GroupThreadID.x; CurrIdx < D; CurrIdx += THREADGROUP_SIZE_X)
+		{
+			Output[StartIdx + CurrIdx] = WRITE(exp(READ(Input[StartIdx + CurrIdx])) / SumExp);
+		}
+	}
+}
\ No newline at end of file
diff --git a/Engine/Plugins/Experimental/NNERuntimeRDG/Source/NNERuntimeRDG/Private/Hlsl/NNERuntimeRDGSoftmax.cpp b/Engine/Plugins/Experimental/NNERuntimeRDG/Source/NNERuntimeRDG/Private/Hlsl/NNERuntimeRDGSoftmax.cpp
new file mode 100644
--- /dev/null
+++ b/Engine/Plugins/Experimental/NNERuntimeRDG/Source/NNERuntimeRDG/Private/Hlsl/NNERuntimeRDGSoftmax.cpp
@@ -0,0 +1,156 @@
+// Copyright Epic Games, Inc. All Rights Reserved.
+
+#include "NNERuntimeRDGSoftmax.h"
+#include "NNEHlslShadersSoftmaxCS.h"
+#include "NNERuntimeRDGHlslHelper.h"
+#include "NNETensor.h"
+#include "NNETypes.h"
+#include "RenderGraphUtils.h"
+
+namespace UE::NNERuntimeRDG::Private::Hlsl
+{
+	DECLARE_GPU_STAT_NAMED(FNNEOperatorSoftmax, TEXT("NNE.Operator.Hlsl.Softmax"));
+
+	/**
+	 * Softmax operator implementation
+	 */
+	class FSoftmax : public FOperatorHlsl
+	{
+
+	public:
+
+		FSoftmax() = default;
+		virtual ~FSoftmax() = default;
+
+	private:
+
+		int32 Axis = 1;
+
+	public:
+
+		virtual int PrepareOutputs(TConstArrayView<NNE::Internal::FTensorRef> InputTensors, TArrayView<NNE::Internal::FTensorRef> OutputTensors) const override
+		{
+			check(InputTensors.Num() == 1);
+			check(OutputTensors.Num() == 1);
+
+			const NNE::FTensorShape& InputShape = InputTensors[0]->GetShape();
+
+			OutputTensors[0]->SetShape(InputShape);
+
+			return 0;
+		};
+
+		virtual bool Initialize(TConstArrayView<NNE::FTensorDesc> InputTensorDescs, TConstArrayView<NNE::FTensorDesc> OutputTensorDescs, const NNE::FAttributeMap& Attributes) override
+		{
+			check(InputTensorDescs.Num() == 1);
+			check(OutputTensorDescs.Num() == 1);
+
+			const int32 InputDimensions = InputTensorDescs[0].GetShape().Rank();
+
+			if (InputDimensions < 2)
+			{
+				UE_LOG(LogNNE, Warning, TEXT("Input tensor should be at least 2-D (but got rank %d)"), InputDimensions);
+				return false;
+			}
+
+			if (InputTensorDescs[0].GetShape().Rank() != OutputTensorDescs[0].GetShape().Rank())
+			{
+				UE_LOG(LogNNE, Warning, TEXT("Softwax requires the output to have the same rank as the input."));
+				return false;
+			}
+
+			const int32 AxisLow = 0;
+			const int32 AxisHigh = InputDimensions - 1;
+
+			Axis = Attributes.GetValueOrDefault<int32>(TEXT("axis"), 1);
+			if (Axis < AxisLow || AxisHigh < Axis)
+			{
+				UE_LOG(LogNNE, Warning, TEXT("Invalid axis (should in the interval [%d, %d], but got %d)"), AxisLow, AxisHigh, Axis);
+				return false;
+			}
+
+			return true;
+		}
+
+		virtual void Dispatch(FRDGBuilder& GraphBuilder, TConstArrayView<FTensorRDGRef> InputTensors, TConstArrayView<FTensorRDGRef> OutputTensors) override
+		{
+			using namespace UE::NNEHlslShaders::Internal;
+
+			check(InputTensors.Num() == 1);
+			check(OutputTensors.Num() == 1);
+			check(InputTensors[0] != FTensorRDGRef{});
+			check(OutputTensors[0] != FTensorRDGRef{});
+
+			const FTensorRDG& Input = *InputTensors[0];
+			const FTensorRDG& Output = *OutputTensors[0];
+
+			check(Input.GetVolume() == Output.GetVolume());
+
+			const NNE::FTensorShape& InputShape = Input.GetShape();
+			TConstArrayView<uint32> InputShapeData = InputShape.GetData();
+			const int32 InputDimensions = InputShape.Rank();
+
+			int32 N = 1;
+			for (int32 i = 0; i < Axis; i++)
+			{
+				N *= InputShapeData[i];
+			}
+
+			int32 D = 1;
+			for (int32 i = Axis; i < InputDimensions; i++)
+			{
+				D *= InputShapeData[i];
+			}
+
+			// Set parameters
+			TSoftmaxCS::FParameters* Parameters = GraphBuilder.AllocParameters<TSoftmaxCS::FParameters>();
+			Parameters->N = N;
+			Parameters->D = D;
+			Parameters->Input = GraphBuilder.CreateSRV(FRDGBufferSRVDesc(Input.GetBuffer(), PF_R32_FLOAT));
+			Parameters->Output = GraphBuilder.CreateUAV(FRDGBufferUAVDesc(Output.GetBuffer(), PF_R32_FLOAT));
+
+			TShaderMapRef<TSoftmaxCS> ComputeShader(GetGlobalShaderMap(GMaxRHIFeatureLevel));
+			FIntVector ThreadGroupCount = TSoftmaxCS::GetGroupCount(*Parameters);
+
+			RDG_EVENT_SCOPE(GraphBuilder, "NNE.Operator.Hlsl.Softmax");
+			RDG_GPU_STAT_SCOPE(GraphBuilder, FNNEOperatorSoftmax);
+
+			FComputeShaderUtils::AddPass(
+				GraphBuilder,
+				RDG_EVENT_NAME("NNE.Operator.Hlsl.Softmax.Dispatch"),
+				ERDGPassFlags::Compute | ERDGPassFlags::NeverCull,
+				ComputeShader,
+				Parameters,
+				ThreadGroupCount);
+		}
+	};
+
+	bool ValidateSoftmaxOperator(const NNE::FAttributeMap& AttributeMap, TConstArrayView<ENNETensorDataType> InputTypes, TConstArrayView<NNE::FSymbolicTensorShape> InputShapes)
+	{
+		bool bIsValid = true;
+
+		// This match version 1 of the Softmax operator, next versions are 11 and 13
+		// [REDACTED_PROFILE_URL]/onnx/blob/main/docs/Changelog.md#Softmax-1
+		FAttributeValidator AttributeValidator;
+		AttributeValidator.AddOptional(TEXT("axis"), ENNEAttributeDataType::Int32);
+		bIsValid &= AttributeValidator.Validate(AttributeMap);
+
+		FInputValidator InputValidator;
+		InputValidator.AddSupportedType(ENNETensorDataType::Float);
+		InputValidator.AddRequired();
+		bIsValid &= InputValidator.Validate(InputTypes);
+
+		return bIsValid;
+	}
+
+	FOperatorHlsl* CreateSoftmaxOperator()
+	{
+		return new FSoftmax();
+	}
+
+	bool RegisterSoftmaxOperator(FOperatorRegistryHlsl& Registry)
+	{
+		Registry.OpAdd({{TEXT("Softmax"), TEXT("Onnx")}}, CreateSoftmaxOperator, ValidateSoftmaxOperator);
+		return true;
+	}
+} // UE::NNERuntimeRDG::Private::Hlsl
diff --git a/Engine/Plugins/Experimental/NNERuntimeRDG/Source/NNERuntimeRDG/Private/Hlsl/NNERuntimeRDGSoftmax.h b/Engine/Plugins/Experimental/NNERuntimeRDG/Source/NNERuntimeRDG/Private/Hlsl/NNERuntimeRDGSoftmax.h
new file mode 100644
--- /dev/null
+++ b/Engine/Plugins/Experimental/NNERuntimeRDG/Source/NNERuntimeRDG/Private/Hlsl/NNERuntimeRDGSoftmax.h
@@ -0,0 +1,10 @@
+// Copyright Epic Games, Inc. All Rights Reserved.
+
+#pragma once
+
+#include "NNERuntimeRDGHlslOp.h"
+
+namespace UE::NNERuntimeRDG::Private::Hlsl
+{
+	bool RegisterSoftmaxOperator(FOperatorRegistryHlsl& Registry);
+} // UE::NNERuntimeRDG::Private::Hlsl
diff --git a/Engine/Plugins/Experimental/NNERuntimeRDG/Source/NNERuntimeRDG/Private/NNERuntimeRDGHlsl.cpp b/Engine/Plugins/Experimental/NNERuntimeRDG/Source/NNERuntimeRDG/Private/NNERuntimeRDGHlsl.cpp
--- a/Engine/Plugins/Experimental/NNERuntimeRDG/Source/NNERuntimeRDG/Private/NNERuntimeRDGHlsl.cpp
+++ b/Engine/Plugins/Experimental/NNERuntimeRDG/Source/NNERuntimeRDG/Private/NNERuntimeRDGHlsl.cpp
@@ -34,6 +34,7 @@
 #include "Hlsl/NNERuntimeRDGShape.h"
 #include "Hlsl/NNERuntimeRDGSize.h"
 #include "Hlsl/NNERuntimeRDGSlice.h"
+#include "Hlsl/NNERuntimeRDGSoftmax.h"
 #include "Hlsl/NNERuntimeRDGSqueeze.h"
 #include "Hlsl/NNERuntimeRDGTranspose.h"
 #include "Hlsl/NNERuntimeRDGUnsqueeze.h"
@@ -47,34 +48,35 @@ int32 UNNERuntimeRDGHlslImpl::Version = 0x00000004;
 
 bool UNNERuntimeRDGHlslImpl::Init()
 {
-	FOperatorRegistryHlsl* registry = FOperatorRegistryHlsl::Get();
-	check(registry != nullptr);
-
-	RegisterBatchNormalizationOperator(*registry);
-	RegisterCastOperator(*registry);
-	RegisterConvOperator(*registry);
-	RegisterConcatOperator(*registry);
-	RegisterConvTransposeOperator(*registry);
-	RegisterDropoutOperator(*registry);
-	RegisterElementWiseBinaryOperators(*registry);
-	RegisterElementWiseUnaryOperators(*registry);
-	RegisterElementWiseVariadicOperators(*registry);
-	RegisterFlattenOperator(*registry);
-	RegisterGatherOperator(*registry);
-	RegisterGemmOperator(*registry);
-	RegisterIdentityOperator(*registry);
-	RegisterInstanceNormalizationOperator(*registry);
-	RegisterPadOperator(*registry);
-	RegisterPoolOperators(*registry);
-	RegisterReshapeOperator(*registry);
-	RegisterShapeOperator(*registry);
-	RegisterSizeOperator(*registry);
-	RegisterSliceOperator(*registry);
-	RegisterSqueezeOperator(*registry);
-	RegisterTransposeOperator(*registry);
-	RegisterUnsqueezeOperator(*registry);
-	RegisterUpsampleOperator(*registry);
-	RegisterMatMulOperator(*registry);
+	FOperatorRegistryHlsl* Registry = FOperatorRegistryHlsl::Get();
+	check(Registry != nullptr);
+
+	RegisterBatchNormalizationOperator(*Registry);
+	RegisterCastOperator(*Registry);
+	RegisterConvOperator(*Registry);
+	RegisterConcatOperator(*Registry);
+	RegisterConvTransposeOperator(*Registry);
+	RegisterDropoutOperator(*Registry);
+	RegisterElementWiseBinaryOperators(*Registry);
+	RegisterElementWiseUnaryOperators(*Registry);
+	RegisterElementWiseVariadicOperators(*Registry);
+	RegisterFlattenOperator(*Registry);
+	RegisterGatherOperator(*Registry);
+	RegisterGemmOperator(*Registry);
+	RegisterIdentityOperator(*Registry);
+	RegisterInstanceNormalizationOperator(*Registry);
+	RegisterPadOperator(*Registry);
+	RegisterPoolOperators(*Registry);
+	RegisterReshapeOperator(*Registry);
+	RegisterShapeOperator(*Registry);
+	RegisterSizeOperator(*Registry);
+	RegisterSliceOperator(*Registry);
+	RegisterSoftmaxOperator(*Registry);
+	RegisterSqueezeOperator(*Registry);
+	RegisterTransposeOperator(*Registry);
+	RegisterUnsqueezeOperator(*Registry);
+	RegisterUpsampleOperator(*Registry);
+	RegisterMatMulOperator(*Registry);
 
 	return true;
 }
