[NNE] NNERuntimeRDGHlsl ReduceSum operator

Problem:
NNERuntimeRDGHlsl was missing ReduceSum operator

Solution:
Implement it

Test:
Win64 Editor: ReduceSum tests + smoke tests
Win64 Standalone: all tests

#rnx

#virtualized

diff --git a/Engine/Plugins/Experimental/NNERuntimeRDG/Source/NNEHlslShaders/Internal/NNEHlslShadersReduceCS.h b/Engine/Plugins/Experimental/NNERuntimeRDG/Source/NNEHlslShaders/Internal/NNEHlslShadersReduceCS.h
new file mode 100644
--- /dev/null
+++ b/Engine/Plugins/Experimental/NNERuntimeRDG/Source/NNEHlslShaders/Internal/NNEHlslShadersReduceCS.h
@@ -0,0 +1,40 @@
+// Copyright Epic Games, Inc. All Rights Reserved.
+
+#pragma once
+
+#include "CoreMinimal.h"
+#include "NNEHlslShaderBase.h"
+#include "RenderGraphFwd.h"
+#include "RenderGraphUtils.h"
+#include "ShaderParameterUtils.h"
+
+namespace UE::NNE::Internal { class FTensor; }
+
+namespace UE::NNEHlslShaders::Internal
+{
+	class FReduceConstants
+	{
+	public:
+		static const int32 NUM_GROUP_THREADS{ 768 };
+	};
+
+	class NNEHLSLSHADERS_API TReduceCS : public FHlslShaderBase
+	{
+		DECLARE_GLOBAL_SHADER(TReduceCS);
+		SHADER_USE_PARAMETER_STRUCT(TReduceCS, FHlslShaderBase)
+
+	public:
+
+		BEGIN_SHADER_PARAMETER_STRUCT(FParameters, )
+			SHADER_PARAMETER(int32, NumElemBeforeAxis)
+			SHADER_PARAMETER(int32, AxisSize)
+			SHADER_PARAMETER(int32, NumElemAfterAxis)
+			SHADER_PARAMETER_RDG_BUFFER_SRV(Buffer<float>, Input)
+			SHADER_PARAMETER_RDG_BUFFER_UAV(RWBuffer<float>, Output)
+		END_SHADER_PARAMETER_STRUCT()
+
+		static void ModifyCompilationEnvironment(const FGlobalShaderPermutationParameters& InParameters, FShaderCompilerEnvironment& OutEnvironment);
+		static void FillInParameters(TConstArrayView<uint32> Shape, int32 Axis, FParameters* Parameters);
+		static void EnqueueRDG(FRDGBuilder& GraphBuilder, FParameters* Parameters, FRDGBufferRef Input, FRDGBufferRef Output);
+	};
+} // UE::NNEHlslShaders::Internal
\ No newline at end of file
diff --git a/Engine/Plugins/Experimental/NNERuntimeRDG/Source/NNEHlslShaders/Private/NNEHlslShadersReduceCS.cpp b/Engine/Plugins/Experimental/NNERuntimeRDG/Source/NNEHlslShaders/Private/NNEHlslShadersReduceCS.cpp
new file mode 100644
--- /dev/null
+++ b/Engine/Plugins/Experimental/NNERuntimeRDG/Source/NNEHlslShaders/Private/NNEHlslShadersReduceCS.cpp
@@ -0,0 +1,57 @@
+// Copyright Epic Games, Inc. All Rights Reserved.
+
+#include "NNEHlslShadersReduceCS.h"
+#include "NNE.h"
+#include "RenderGraphBuilder.h"
+
+namespace UE::NNEHlslShaders::Internal
+{
+	void TReduceCS::ModifyCompilationEnvironment(const FGlobalShaderPermutationParameters& InParameters, FShaderCompilerEnvironment& OutEnvironment)
+	{
+		FGlobalShader::ModifyCompilationEnvironment(InParameters, OutEnvironment);
+
+		OutEnvironment.SetDefine(TEXT("THREADGROUP_SIZE"), FReduceConstants::NUM_GROUP_THREADS);
+	}
+
+	void TReduceCS::FillInParameters(TConstArrayView<uint32> Shape, int32 Axis, TReduceCS::FParameters* Parameters)
+	{
+		check(Axis >= 0 && Axis <= Shape.Num());
+		check(Parameters);
+
+		int32 NumElemBeforeAxis = 1;
+		int32 NumElemAfterAxis = 1;
+		for (int32 d = 0; d < Axis; ++d)
+		{
+			NumElemBeforeAxis *= Shape[d];
+		}
+		for (int32 d = Axis + 1; d < Shape.Num(); ++d)
+		{
+			NumElemAfterAxis *= Shape[d];
+		}
+
+		Parameters->NumElemBeforeAxis = NumElemBeforeAxis;
+		Parameters->AxisSize = Shape[Axis];
+		Parameters->NumElemAfterAxis = NumElemAfterAxis;
+	}
+
+	void TReduceCS::EnqueueRDG(FRDGBuilder& GraphBuilder, TReduceCS::FParameters* Parameters, FRDGBufferRef Input, FRDGBufferRef Output)
+	{
+		check(Parameters);
+
+		Parameters->Input = GraphBuilder.CreateSRV(FRDGBufferSRVDesc(Input, PF_R32_FLOAT));
+		Parameters->Output = GraphBuilder.CreateUAV(FRDGBufferUAVDesc(Output, PF_R32_FLOAT));
+
+		TShaderMapRef<TReduceCS> ComputeShader(GetGlobalShaderMap(GMaxRHIFeatureLevel));
+		FIntVector ThreadGroupCount{ 1, (int32)Parameters->NumElemAfterAxis, (int32)Parameters->NumElemBeforeAxis };
+
+		FComputeShaderUtils::AddPass(
+			GraphBuilder,
+			RDG_EVENT_NAME("NNE.Operator.Hlsl.Reduce.OneAxis.Dispatch"),
+			ERDGPassFlags::Compute | ERDGPassFlags::NeverCull,
+			ComputeShader,
+			Parameters,
+			ThreadGroupCount);
+	}
+
+	IMPLEMENT_GLOBAL_SHADER(TReduceCS, "/NNE/NNEHlslShadersReduce.usf", "Reduce", SF_Compute);
+} // UE::NNEHlslShaders::Internal
diff --git a/Engine/Plugins/Experimental/NNERuntimeRDG/Source/NNEHlslShaders/Shaders/NNEHlslShadersReduce.usf b/Engine/Plugins/Experimental/NNERuntimeRDG/Source/NNEHlslShaders/Shaders/NNEHlslShadersReduce.usf
new file mode 100644
--- /dev/null
+++ b/Engine/Plugins/Experimental/NNERuntimeRDG/Source/NNEHlslShaders/Shaders/NNEHlslShadersReduce.usf
@@ -0,0 +1,63 @@
+// Copyright Epic Games, Inc. All Rights Reserved.
+
+#include "/Engine/Public/Platform.ush"
+
+#define WORK_TYPE float
+#define BUFFER_TYPE float
+#define READ(x) x
+#define WRITE(x) x
+
+Buffer<WORK_TYPE> Input;
+RWBuffer<WORK_TYPE> Output;
+int NumElemBeforeAxis;
+int AxisSize;
+int NumElemAfterAxis;
+
+groupshared WORK_TYPE SharedMemorySum[THREADGROUP_SIZE];
+
+[numthreads(THREADGROUP_SIZE, 1, 1)]
+void Reduce(in const uint3 DispatchThreadID : SV_DispatchThreadID, in const uint3 GroupThreadID : SV_GroupThreadID)
+{
+	const uint IdxAfterAxis = DispatchThreadID.y;
+	const uint IdxBeforeAxis = DispatchThreadID.z;
+	const int InputStartIdx = (IdxBeforeAxis * AxisSize * NumElemAfterAxis) + IdxAfterAxis;
+	const int OutputIdx     = (IdxBeforeAxis * 1        * NumElemAfterAxis) + IdxAfterAxis;
+
+	// Reduction using sharedmem and THREADGROUP_SIZE threads along the reduced dimension
+	{
+		WORK_TYPE Sum = 0;
+
+		for (int AxisIdx = GroupThreadID.x; AxisIdx < AxisSize; AxisIdx += THREADGROUP_SIZE)
+		{
+			WORK_TYPE Value = READ(Input[(AxisIdx * NumElemAfterAxis) + InputStartIdx]);
+			Sum += Value;
+		}
+
+		SharedMemorySum[GroupThreadID.x] = Sum;
+	}
+
+	// Wait until all threads wrote result to shared memory
+	GroupMemoryBarrierWithGroupSync();
+
+	// Use one thread to calculate final result
+	if (GroupThreadID.x == 0)
+	{
+		WORK_TYPE Sum = 0;
+
+		UNROLL
+		for (int Idx = 0; Idx < THREADGROUP_SIZE; Idx++)
+		{
+			WORK_TYPE CurrSum = SharedMemorySum[Idx];
+			Sum += CurrSum;
+		}
+
+		SharedMemorySum[0] = Sum;
+	}
+
+	// Wait until first thread wrote result to shared memory
+	GroupMemoryBarrierWithGroupSync();
+
+	// Write reduced result
+	WORK_TYPE Sum = SharedMemorySum[0];
+	Output[OutputIdx] = Sum;
+}
\ No newline at end of file
diff --git a/Engine/Plugins/Experimental/NNERuntimeRDG/Source/NNERuntimeRDG/Private/Hlsl/NNERuntimeRDGReduce.cpp b/Engine/Plugins/Experimental/NNERuntimeRDG/Source/NNERuntimeRDG/Private/Hlsl/NNERuntimeRDGReduce.cpp
new file mode 100644
--- /dev/null
+++ b/Engine/Plugins/Experimental/NNERuntimeRDG/Source/NNERuntimeRDG/Private/Hlsl/NNERuntimeRDGReduce.cpp
@@ -0,0 +1,171 @@
+// Copyright Epic Games, Inc. All Rights Reserved.
+
+#include "NNERuntimeRDGReduce.h"
+#include "NNEHlslShadersReduceCS.h"
+#include "NNERuntimeRDGHlslHelper.h"
+#include "NNETensor.h"
+#include "NNETypes.h"
+#include "RenderGraphUtils.h"
+
+namespace UE::NNERuntimeRDG::Private::Hlsl
+{
+	DECLARE_GPU_STAT_NAMED(FNNEOperatorReduceSum, TEXT("NNE.Operator.Hlsl.ReduceSum"));
+
+	/**
+	 * ReduceSum operator implementation
+	 */
+	class FReduceSum : public FOperatorHlsl
+	{
+
+	public:
+
+		FReduceSum() = default;
+		virtual ~FReduceSum() = default;
+
+	private:
+
+		TArray<int32> Axes;
+		int32 KeepDims = 1;
+
+	public:
+
+		virtual int PrepareOutputs(TConstArrayView<NNE::Internal::FTensorRef> InputTensors, TArrayView<NNE::Internal::FTensorRef> OutputTensors) const override
+		{
+			check(InputTensors.Num() == 1);
+			check(OutputTensors.Num() == 1);
+
+			const NNE::FTensorShape& InputShape = InputTensors[0]->GetShape();
+			const int32 InputRank = InputShape.Rank();
+			TArray<uint32> OutputShape;
+
+			OutputShape.Reserve(InputRank);
+			for (int i = 0; i < InputRank; ++i)
+			{
+				if (!Axes.Contains(i))
+				{
+					OutputShape.Add(InputShape.GetData()[i]);
+				}
+				else if (KeepDims)
+				{
+					OutputShape.Add(1);
+				}
+			}
+
+			OutputTensors[0]->SetShape(NNE::FTensorShape::Make(OutputShape));
+
+			return 0;
+		};
+
+		virtual bool Initialize(TConstArrayView<NNE::FTensorDesc> InputTensorDescs, TConstArrayView<NNE::FTensorDesc> OutputTensorDescs, const NNE::FAttributeMap& Attributes) override
+		{
+			check(InputTensorDescs.Num() == 1);
+			check(OutputTensorDescs.Num() == 1);
+
+			const int32 InputRank = InputTensorDescs[0].GetShape().Rank();
+			TArray<int32> AxesTemp;
+
+			for (int i = 0; i < InputRank; ++i)
+			{
+				AxesTemp.Add(i);
+			}
+			AxesTemp = Attributes.GetValueOrDefault<TArray<int32>>(TEXT("axes"), AxesTemp);
+			AxesTemp.Sort();
+
+			Axes.Empty();
+
+			for (int32 axis : AxesTemp)
+			{
+				if (axis > InputRank || axis  < -InputRank)
+				{
+					UE_LOG(LogNNE, Warning, TEXT("ReduceSum 'Axes' attribute should contain value be in the range [-r,r] with r being the rank of the input (name: %s) however got %d while rank is %d."), *InputTensorDescs[0].GetName(), axis, InputRank);
+					return false;
+				}
+
+				if (axis  < 0)
+				{
+					axis = InputRank + axis;
+				}
+				Axes.AddUnique(axis);
+			}
+			check(Axes.Num() > 0);
+
+			KeepDims = Attributes.GetValueOrDefault<int32>(TEXT("keepdims"), 1);
+
+			return true;
+		}
+
+		virtual void Dispatch(FRDGBuilder& GraphBuilder, TConstArrayView<FTensorRDGRef> InputTensors, TConstArrayView<FTensorRDGRef> OutputTensors) override
+		{
+			using namespace UE::NNEHlslShaders::Internal;
+
+			check(InputTensors.Num() == 1);
+			check(OutputTensors.Num() == 1);
+			check(InputTensors[0] != FTensorRDGRef{});
+			check(OutputTensors[0] != FTensorRDGRef{});
+			check(Axes.Num() > 0);
+
+			const FTensorRDG& Input = *InputTensors[0];
+			const FTensorRDG& Output = *OutputTensors[0];
+			const int32 InputRank = Input.GetShape().Rank();
+
+			RDG_EVENT_SCOPE(GraphBuilder, "NNE.Operator.Hlsl.ReduceSum");
+			RDG_GPU_STAT_SCOPE(GraphBuilder, FNNEOperatorReduceSum);
+
+			FRDGBufferRef CurrInput = Input.GetBuffer();
+			TArray<uint32, TInlineAllocator<NNE::FTensorShape::MaxRank>> CurrInputShape(Input.GetShape().GetData());
+			for (int32 i = Axes.Num()-1; i >= 0; --i)
+			{
+				const int32 Axis = Axes[i];
+
+				TReduceCS::FParameters* Parameters = GraphBuilder.AllocParameters<TReduceCS::FParameters>();
+				TReduceCS::FillInParameters(CurrInputShape, Axis, Parameters);
+
+				FRDGBufferRef CurrOutput = nullptr;
+				if (i != 0)
+				{
+					const FRDGBufferDesc TempBufferDesc = FRDGBufferDesc::CreateBufferDesc(Output.GetElementByteSize(), Parameters->NumElemBeforeAxis * Parameters->NumElemAfterAxis);
+					CurrOutput = GraphBuilder.CreateBuffer(TempBufferDesc, TEXT("NNE.Operator.Hlsl.ReduceSum.TempBuffer"), ERDGBufferFlags::None);
+				}
+				else
+				{
+					CurrOutput = Output.GetBuffer();
+				}
+				check(CurrOutput);
+
+				TReduceCS::EnqueueRDG(GraphBuilder, Parameters, CurrInput, CurrOutput);
+				CurrInput = CurrOutput;
+				CurrInputShape[i] = 1;
+			}
+		}
+	};
+
+	bool ValidateReduceSumOperator(const NNE::FAttributeMap& AttributeMap, TConstArrayView<ENNETensorDataType> InputTypes, TConstArrayView<NNE::FSymbolicTensorShape> InputShapes)
+	{
+		bool bIsValid = true;
+
+		// This match version 1 and 11 of the ReduceSum operator, next version is 13
+		// [REDACTED_PROFILE_URL]/onnx/blob/main/docs/Changelog.md#Reduce-1
+		FAttributeValidator AttributeValidator;
+		AttributeValidator.AddOptional(TEXT("axes"), ENNEAttributeDataType::Int32Array);
+		AttributeValidator.AddOptional(TEXT("keepdims"), ENNEAttributeDataType::Int32);
+		bIsValid &= AttributeValidator.Validate(AttributeMap);
+
+		FInputValidator InputValidator;
+		InputValidator.AddSupportedType(ENNETensorDataType::Float);
+		InputValidator.AddRequired();
+		bIsValid &= InputValidator.Validate(InputTypes);
+
+		return bIsValid;
+	}
+
+	FOperatorHlsl* CreateReduceSumOperator()
+	{
+		return new FReduceSum();
+	}
+
+	bool RegisterReduceOperators(FOperatorRegistryHlsl& Registry)
+	{
+		Registry.OpAdd({{TEXT("ReduceSum"), TEXT("Onnx")}}, CreateReduceSumOperator, ValidateReduceSumOperator);
+		return true;
+	}
+} // UE::NNERuntimeRDG::Private::Hlsl
diff --git a/Engine/Plugins/Experimental/NNERuntimeRDG/Source/NNERuntimeRDG/Private/Hlsl/NNERuntimeRDGReduce.h b/Engine/Plugins/Experimental/NNERuntimeRDG/Source/NNERuntimeRDG/Private/Hlsl/NNERuntimeRDGReduce.h
new file mode 100644
--- /dev/null
+++ b/Engine/Plugins/Experimental/NNERuntimeRDG/Source/NNERuntimeRDG/Private/Hlsl/NNERuntimeRDGReduce.h
@@ -0,0 +1,10 @@
+// Copyright Epic Games, Inc. All Rights Reserved.
+
+#pragma once
+
+#include "NNERuntimeRDGHlslOp.h"
+
+namespace UE::NNERuntimeRDG::Private::Hlsl
+{
+	bool RegisterReduceOperators(FOperatorRegistryHlsl& Registry);
+} // UE::NNERuntimeRDG::Private::Hlsl
diff --git a/Engine/Plugins/Experimental/NNERuntimeRDG/Source/NNERuntimeRDG/Private/NNERuntimeRDGHlsl.cpp b/Engine/Plugins/Experimental/NNERuntimeRDG/Source/NNERuntimeRDG/Private/NNERuntimeRDGHlsl.cpp
--- a/Engine/Plugins/Experimental/NNERuntimeRDG/Source/NNERuntimeRDG/Private/NNERuntimeRDGHlsl.cpp
+++ b/Engine/Plugins/Experimental/NNERuntimeRDG/Source/NNERuntimeRDG/Private/NNERuntimeRDGHlsl.cpp
@@ -30,6 +30,7 @@
 #include "Hlsl/NNERuntimeRDGInstanceNormalization.h"
 #include "Hlsl/NNERuntimeRDGPad.h"
 #include "Hlsl/NNERuntimeRDGPool.h"
+#include "Hlsl/NNERuntimeRDGReduce.h"
 #include "Hlsl/NNERuntimeRDGReshape.h"
 #include "Hlsl/NNERuntimeRDGShape.h"
 #include "Hlsl/NNERuntimeRDGSize.h"
@@ -67,6 +68,7 @@ bool UNNERuntimeRDGHlslImpl::Init()
 	RegisterInstanceNormalizationOperator(*Registry);
 	RegisterPadOperator(*Registry);
 	RegisterPoolOperators(*Registry);
+	RegisterReduceOperators(*Registry);
 	RegisterReshapeOperator(*Registry);
 	RegisterShapeOperator(*Registry);
 	RegisterSizeOperator(*Registry);
