[NNE] Bump RDG ModelData versions

Problem:
Current DDC cache entries cannot be unserialized by RDG ModelData, resulting in crashes when using RDG runtimes.

Workaround/solution:
Bump RDG ModelData versions to v5 so RDG ModelData are recreated and stored in the DDC while we continue to investigate why the v4 DDC entries cannot be unserialized anymore. Unblocking the team/use cases using RDG runtimes.

Test: Win64 Editor: smoke tests + softmax tests
Test: Win64 Standalone : all tests

#rnx

diff --git a/Engine/Plugins/Experimental/NNERuntimeRDG/Source/NNERuntimeRDG/Private/Dml/NNEDmlModel.cpp b/Engine/Plugins/Experimental/NNERuntimeRDG/Source/NNERuntimeRDG/Private/Dml/NNEDmlModel.cpp
--- a/Engine/Plugins/Experimental/NNERuntimeRDG/Source/NNERuntimeRDG/Private/Dml/NNEDmlModel.cpp
+++ b/Engine/Plugins/Experimental/NNERuntimeRDG/Source/NNERuntimeRDG/Private/Dml/NNEDmlModel.cpp
@@ -44,7 +44,7 @@ FModelInfo* FModelInfo::Get()
 
 FModelInfo::FModelInfo()
 	: Guid((int32)'R', (int32)'D', (int32)'G', (int32)'D')
-	, Version(0x4)
+	, Version(0x5)
 {
 }
 
diff --git a/Engine/Plugins/Experimental/NNERuntimeRDG/Source/NNERuntimeRDG/Private/NNERuntimeRDGHlsl.cpp b/Engine/Plugins/Experimental/NNERuntimeRDG/Source/NNERuntimeRDG/Private/NNERuntimeRDGHlsl.cpp
--- a/Engine/Plugins/Experimental/NNERuntimeRDG/Source/NNERuntimeRDG/Private/NNERuntimeRDGHlsl.cpp
+++ b/Engine/Plugins/Experimental/NNERuntimeRDG/Source/NNERuntimeRDG/Private/NNERuntimeRDGHlsl.cpp
@@ -45,7 +45,7 @@
 using namespace UE::NNERuntimeRDG::Private::Hlsl;
 
 FGuid UNNERuntimeRDGHlslImpl::GUID = FGuid((int32)'R', (int32)'D', (int32)'G', (int32)'H');
-int32 UNNERuntimeRDGHlslImpl::Version = 0x00000004;
+int32 UNNERuntimeRDGHlslImpl::Version = 0x00000005;
 
 bool UNNERuntimeRDGHlslImpl::Init()
 {
