[NNE] NNERuntimeHLSL add variations to Reduce operators

Problem:
Only ReduceSum operator was available however in order to factorize the reduce operations to support Softmax (need sumexp) and GlobalAveragePool (need average) we needed more reduction type.

Solution:
Add Sumexp and Average variation to Reduction type, while at it add low hanging fruits other types as Reduce operator (meaning we get tests coverage for this too)

Test: Win64 Editor: smoke tests + reduce tests
Test: Win64 Standalone : all tests

#rnx

#virtualized

diff --git a/Engine/Plugins/Experimental/NNERuntimeRDG/Source/NNEHlslShaders/Internal/NNEHlslShadersReduceCS.h b/Engine/Plugins/Experimental/NNERuntimeRDG/Source/NNEHlslShaders/Internal/NNEHlslShadersReduceCS.h
--- a/Engine/Plugins/Experimental/NNERuntimeRDG/Source/NNEHlslShaders/Internal/NNEHlslShadersReduceCS.h
+++ b/Engine/Plugins/Experimental/NNERuntimeRDG/Source/NNEHlslShaders/Internal/NNEHlslShadersReduceCS.h
@@ -12,6 +12,20 @@ namespace UE::NNE::Internal { class FTensor; }
 
 namespace UE::NNEHlslShaders::Internal
 {
+	enum class EReduceOperatorType : uint8
+	{
+		Average = 0,
+		L1,
+		L2,
+		LogSumExp,
+		Max,
+		Min,
+		Prod,
+		Sum,
+		SumExp,//Should not be used for multiple axis reduction
+		MAX
+	};
+
 	class FReduceConstants
 	{
 	public:
@@ -23,6 +37,9 @@ namespace UE::NNEHlslShaders::Internal
 		DECLARE_GLOBAL_SHADER(TReduceCS);
 		SHADER_USE_PARAMETER_STRUCT(TReduceCS, FHlslShaderBase)
 
+		class FReduceType : SHADER_PERMUTATION_ENUM_CLASS("REDUCE_OPERATOR_TYPE", EReduceOperatorType);
+		using FPermutationDomain = TShaderPermutationDomain<FReduceType>;
+
 	public:
 
 		BEGIN_SHADER_PARAMETER_STRUCT(FParameters, )
@@ -35,6 +52,6 @@ namespace UE::NNEHlslShaders::Internal
 
 		static void ModifyCompilationEnvironment(const FGlobalShaderPermutationParameters& InParameters, FShaderCompilerEnvironment& OutEnvironment);
 		static void FillInParameters(TConstArrayView<uint32> Shape, int32 Axis, FParameters* Parameters);
-		static void EnqueueRDG(FRDGBuilder& GraphBuilder, FParameters* Parameters, FRDGBufferRef Input, FRDGBufferRef Output);
+		static void EnqueueRDG(FRDGBuilder& GraphBuilder, FParameters* Parameters, FRDGBufferRef Input, FRDGBufferRef Output, EReduceOperatorType OperatorType);
 	};
 } // UE::NNEHlslShaders::Internal
\ No newline at end of file
diff --git a/Engine/Plugins/Experimental/NNERuntimeRDG/Source/NNEHlslShaders/Private/NNEHlslShadersReduceCS.cpp b/Engine/Plugins/Experimental/NNERuntimeRDG/Source/NNEHlslShaders/Private/NNEHlslShadersReduceCS.cpp
--- a/Engine/Plugins/Experimental/NNERuntimeRDG/Source/NNEHlslShaders/Private/NNEHlslShadersReduceCS.cpp
+++ b/Engine/Plugins/Experimental/NNERuntimeRDG/Source/NNEHlslShaders/Private/NNEHlslShadersReduceCS.cpp
@@ -34,14 +34,17 @@ namespace UE::NNEHlslShaders::Internal
 		Parameters->NumElemAfterAxis = NumElemAfterAxis;
 	}
 
-	void TReduceCS::EnqueueRDG(FRDGBuilder& GraphBuilder, TReduceCS::FParameters* Parameters, FRDGBufferRef Input, FRDGBufferRef Output)
+	void TReduceCS::EnqueueRDG(FRDGBuilder& GraphBuilder, TReduceCS::FParameters* Parameters, FRDGBufferRef Input, FRDGBufferRef Output, EReduceOperatorType OperatorType)
 	{
 		check(Parameters);
 
 		Parameters->Input = GraphBuilder.CreateSRV(FRDGBufferSRVDesc(Input, PF_R32_FLOAT));
 		Parameters->Output = GraphBuilder.CreateUAV(FRDGBufferUAVDesc(Output, PF_R32_FLOAT));
 
-		TShaderMapRef<TReduceCS> ComputeShader(GetGlobalShaderMap(GMaxRHIFeatureLevel));
+		TReduceCS::FPermutationDomain PermutationVector;
+		PermutationVector.Set<TReduceCS::FReduceType>(OperatorType);
+
+		TShaderMapRef<TReduceCS> ComputeShader(GetGlobalShaderMap(GMaxRHIFeatureLevel), PermutationVector);
 		FIntVector ThreadGroupCount{ 1, (int32)Parameters->NumElemAfterAxis, (int32)Parameters->NumElemBeforeAxis };
 
 		FComputeShaderUtils::AddPass(
diff --git a/Engine/Plugins/Experimental/NNERuntimeRDG/Source/NNEHlslShaders/Shaders/NNEHlslShadersReduce.usf b/Engine/Plugins/Experimental/NNERuntimeRDG/Source/NNEHlslShaders/Shaders/NNEHlslShadersReduce.usf
--- a/Engine/Plugins/Experimental/NNERuntimeRDG/Source/NNEHlslShaders/Shaders/NNEHlslShadersReduce.usf
+++ b/Engine/Plugins/Experimental/NNERuntimeRDG/Source/NNEHlslShaders/Shaders/NNEHlslShadersReduce.usf
@@ -13,7 +13,110 @@ int NumElemBeforeAxis;
 int AxisSize;
 int NumElemAfterAxis;
 
-groupshared WORK_TYPE SharedMemorySum[THREADGROUP_SIZE];
+// Must correspond to EReduceOperatorType defined in NNEHlslShadersReduceCS.h
+#define TYPE_AVERAGE 0
+#define TYPE_L1 1
+#define TYPE_L2 2
+#define TYPE_LOGSUMEXP 3
+#define TYPE_MAX 4
+#define TYPE_MIN 5
+#define TYPE_PROD 6
+#define TYPE_SUM 7
+#define TYPE_SUMEXP 8
+
+WORK_TYPE InitValue(WORK_TYPE FirstValue)
+{
+	#if REDUCE_OPERATOR_TYPE == TYPE_AVERAGE
+		return 0;
+	#elif REDUCE_OPERATOR_TYPE == TYPE_L1
+		return 0;
+	#elif REDUCE_OPERATOR_TYPE == TYPE_L2
+		return 0;
+	#elif REDUCE_OPERATOR_TYPE == TYPE_LOGSUMEXP
+		return 0;
+	#elif REDUCE_OPERATOR_TYPE == TYPE_MAX
+		return FirstValue;
+	#elif REDUCE_OPERATOR_TYPE == TYPE_MIN
+		return FirstValue;
+	#elif REDUCE_OPERATOR_TYPE == TYPE_PROD
+		return 1;
+	#elif REDUCE_OPERATOR_TYPE == TYPE_SUM
+		return 0;
+	#elif REDUCE_OPERATOR_TYPE == TYPE_SUMEXP
+		return 0;
+	#endif
+}
+
+WORK_TYPE ElementOp(WORK_TYPE Value)
+{
+	#if REDUCE_OPERATOR_TYPE == TYPE_AVERAGE
+		return Value;
+	#elif REDUCE_OPERATOR_TYPE == TYPE_L1
+		return abs(Value);
+	#elif REDUCE_OPERATOR_TYPE == TYPE_L2
+		return Value*Value;
+	#elif REDUCE_OPERATOR_TYPE == TYPE_LOGSUMEXP
+		return exp(Value);
+	#elif REDUCE_OPERATOR_TYPE == TYPE_MAX
+		return Value;
+	#elif REDUCE_OPERATOR_TYPE == TYPE_MIN
+		return Value;
+	#elif REDUCE_OPERATOR_TYPE == TYPE_PROD
+		return Value;
+	#elif REDUCE_OPERATOR_TYPE == TYPE_SUM
+		return Value;
+	#elif REDUCE_OPERATOR_TYPE == TYPE_SUMEXP
+		return exp(Value);
+	#endif
+}
+
+WORK_TYPE CombineOp(WORK_TYPE ValueA, WORK_TYPE ValueB)
+{
+	#if REDUCE_OPERATOR_TYPE == TYPE_AVERAGE
+		return ValueA+ValueB;
+	#elif REDUCE_OPERATOR_TYPE == TYPE_L1
+		return ValueA+ValueB;
+	#elif REDUCE_OPERATOR_TYPE == TYPE_L2
+		return ValueA+ValueB;
+	#elif REDUCE_OPERATOR_TYPE == TYPE_LOGSUMEXP
+		return ValueA+ValueB;
+	#elif REDUCE_OPERATOR_TYPE == TYPE_MAX
+		return max(ValueA, ValueB);
+	#elif REDUCE_OPERATOR_TYPE == TYPE_MIN
+		return min(ValueA, ValueB);
+	#elif REDUCE_OPERATOR_TYPE == TYPE_PROD
+		return ValueA*ValueB;
+	#elif REDUCE_OPERATOR_TYPE == TYPE_SUM
+		return ValueA+ValueB;
+	#elif REDUCE_OPERATOR_TYPE == TYPE_SUMEXP
+		return ValueA+ValueB;
+	#endif
+}
+
+WORK_TYPE FinalOp(WORK_TYPE Value)
+{
+	#if REDUCE_OPERATOR_TYPE == TYPE_AVERAGE
+		return Value/AxisSize;
+	#elif REDUCE_OPERATOR_TYPE == TYPE_L1
+		return Value;
+	#elif REDUCE_OPERATOR_TYPE == TYPE_L2
+		return sqrt(Value);
+	#elif REDUCE_OPERATOR_TYPE == TYPE_LOGSUMEXP
+		return log(Value);
+	#elif REDUCE_OPERATOR_TYPE == TYPE_MAX
+		return Value;
+	#elif REDUCE_OPERATOR_TYPE == TYPE_MIN
+		return Value;
+	#elif REDUCE_OPERATOR_TYPE == TYPE_PROD
+		return Value;
+	#elif REDUCE_OPERATOR_TYPE == TYPE_SUM
+		return Value;
+	#elif REDUCE_OPERATOR_TYPE == TYPE_SUMEXP
+		return Value;
+	#endif
+}
+
+groupshared WORK_TYPE SharedMemory[THREADGROUP_SIZE];
 
 [numthreads(THREADGROUP_SIZE, 1, 1)]
 void Reduce(in const uint3 DispatchThreadID : SV_DispatchThreadID, in const uint3 GroupThreadID : SV_GroupThreadID)
@@ -25,15 +128,18 @@ void Reduce(in const uint3 DispatchThreadID : SV_DispatchThreadID, in const uint
 
 	// Reduction using sharedmem and THREADGROUP_SIZE threads along the reduced dimension
 	{
-		WORK_TYPE Sum = 0;
+		WORK_TYPE FirstValue = READ(Input[InputStartIdx]);
+		WORK_TYPE Reduced = InitValue(FirstValue);
 
 		for (int AxisIdx = GroupThreadID.x; AxisIdx < AxisSize; AxisIdx += THREADGROUP_SIZE)
 		{
 			WORK_TYPE Value = READ(Input[(AxisIdx * NumElemAfterAxis) + InputStartIdx]);
-			Sum += Value;
+
+			Value = ElementOp(Value);
+			Reduced = CombineOp(Reduced, Value);
 		}
 
-		SharedMemorySum[GroupThreadID.x] = Sum;
+		SharedMemory[GroupThreadID.x] = Reduced;
 	}
 
 	// Wait until all threads wrote result to shared memory
@@ -42,22 +148,25 @@ void Reduce(in const uint3 DispatchThreadID : SV_DispatchThreadID, in const uint
 	// Use one thread to calculate final result
 	if (GroupThreadID.x == 0)
 	{
-		WORK_TYPE Sum = 0;
+		WORK_TYPE FirstCurrReduced = SharedMemory[0];
+		WORK_TYPE Reduced = InitValue(FirstCurrReduced);
 
 		UNROLL
 		for (int Idx = 0; Idx < THREADGROUP_SIZE; Idx++)
 		{
-			WORK_TYPE CurrSum = SharedMemorySum[Idx];
-			Sum += CurrSum;
+			WORK_TYPE CurrReduced = SharedMemory[Idx];
+
+			Reduced = CombineOp(Reduced, CurrReduced);
 		}
 
-		SharedMemorySum[0] = Sum;
+		SharedMemory[0] = Reduced;
 	}
 
 	// Wait until first thread wrote result to shared memory
 	GroupMemoryBarrierWithGroupSync();
 
 	// Write reduced result
-	WORK_TYPE Sum = SharedMemorySum[0];
-	Output[OutputIdx] = Sum;
+	WORK_TYPE Reduced = SharedMemory[0];
+	Reduced = FinalOp(Reduced);
+	Output[OutputIdx] = Reduced;
 }
\ No newline at end of file
diff --git a/Engine/Plugins/Experimental/NNERuntimeRDG/Source/NNERuntimeRDG/Private/Hlsl/NNERuntimeRDGReduce.cpp b/Engine/Plugins/Experimental/NNERuntimeRDG/Source/NNERuntimeRDG/Private/Hlsl/NNERuntimeRDGReduce.cpp
--- a/Engine/Plugins/Experimental/NNERuntimeRDG/Source/NNERuntimeRDG/Private/Hlsl/NNERuntimeRDGReduce.cpp
+++ b/Engine/Plugins/Experimental/NNERuntimeRDG/Source/NNERuntimeRDG/Private/Hlsl/NNERuntimeRDGReduce.cpp
@@ -9,18 +9,19 @@
 
 namespace UE::NNERuntimeRDG::Private::Hlsl
 {
-	DECLARE_GPU_STAT_NAMED(FNNEOperatorReduceSum, TEXT("NNE.Operator.Hlsl.ReduceSum"));
+	DECLARE_GPU_STAT_NAMED(FNNEOperatorReduce, TEXT("NNE.Operator.Hlsl.Reduce"));
 
 	/**
-	 * ReduceSum operator implementation
+	 * Reduce operators implementation
 	 */
-	class FReduceSum : public FOperatorHlsl
+	template< UE::NNEHlslShaders::Internal::EReduceOperatorType ReduceOperatorType >
+	class FReduceOperator : public FOperatorHlsl
 	{
 
 	public:
 
-		FReduceSum() = default;
-		virtual ~FReduceSum() = default;
+		FReduceOperator() = default;
+		virtual ~FReduceOperator() = default;
 
 	private:
 
@@ -77,7 +78,7 @@ namespace UE::NNERuntimeRDG::Private::Hlsl
 			{
 				if (axis > InputRank || axis  < -InputRank)
 				{
-					UE_LOG(LogNNE, Warning, TEXT("ReduceSum 'Axes' attribute should contain value be in the range [-r,r] with r being the rank of the input (name: %s) however got %d while rank is %d."), *InputTensorDescs[0].GetName(), axis, InputRank);
+					UE_LOG(LogNNE, Warning, TEXT("Reduce operators 'Axes' attribute should contain value be in the range [-r,r] with r being the rank of the input (name: %s) however got %d while rank is %d."), *InputTensorDescs[0].GetName(), axis, InputRank);
 					return false;
 				}
 
@@ -108,8 +109,8 @@ namespace UE::NNERuntimeRDG::Private::Hlsl
 			const FTensorRDG& Output = *OutputTensors[0];
 			const int32 InputRank = Input.GetShape().Rank();
 
-			RDG_EVENT_SCOPE(GraphBuilder, "NNE.Operator.Hlsl.ReduceSum");
-			RDG_GPU_STAT_SCOPE(GraphBuilder, FNNEOperatorReduceSum);
+			RDG_EVENT_SCOPE(GraphBuilder, "NNE.Operator.Hlsl.Reduce");
+			RDG_GPU_STAT_SCOPE(GraphBuilder, FNNEOperatorReduce);
 
 			FRDGBufferRef CurrInput = Input.GetBuffer();
 			TArray<uint32, TInlineAllocator<NNE::FTensorShape::MaxRank>> CurrInputShape(Input.GetShape().GetData());
@@ -124,7 +125,7 @@ namespace UE::NNERuntimeRDG::Private::Hlsl
 				if (i != 0)
 				{
 					const FRDGBufferDesc TempBufferDesc = FRDGBufferDesc::CreateBufferDesc(Output.GetElementByteSize(), Parameters->NumElemBeforeAxis * Parameters->NumElemAfterAxis);
-					CurrOutput = GraphBuilder.CreateBuffer(TempBufferDesc, TEXT("NNE.Operator.Hlsl.ReduceSum.TempBuffer"), ERDGBufferFlags::None);
+					CurrOutput = GraphBuilder.CreateBuffer(TempBufferDesc, TEXT("NNE.Operator.Hlsl.Reduce.TempBuffer"), ERDGBufferFlags::None);
 				}
 				else
 				{
@@ -132,18 +133,19 @@ namespace UE::NNERuntimeRDG::Private::Hlsl
 				}
 				check(CurrOutput);
 
-				TReduceCS::EnqueueRDG(GraphBuilder, Parameters, CurrInput, CurrOutput);
+				TReduceCS::EnqueueRDG(GraphBuilder, Parameters, CurrInput, CurrOutput, ReduceOperatorType);
 				CurrInput = CurrOutput;
 				CurrInputShape[i] = 1;
 			}
 		}
 	};
 
-	bool ValidateReduceSumOperator(const NNE::FAttributeMap& AttributeMap, TConstArrayView<ENNETensorDataType> InputTypes, TConstArrayView<NNE::FSymbolicTensorShape> InputShapes)
+	bool ValidateReduceOperator(const NNE::FAttributeMap& AttributeMap, TConstArrayView<ENNETensorDataType> InputTypes, TConstArrayView<NNE::FSymbolicTensorShape> InputShapes)
 	{
 		bool bIsValid = true;
 
-		// This match version 1 and 11 of the ReduceSum operator, next version is 13
+		// This match versions 1 and 11 of the Reduces operators, next versions are 13
+		// ReduceOperator are ReduceL1, ReduceL2, ReduceLogSum, ReduceLogSumExp, ReduceMax, ReduceMean, ReduceMin, ReduceProd, ReduceSum, ReduceSumSquare
 		// [REDACTED_PROFILE_URL]/onnx/blob/main/docs/Changelog.md#Reduce-1
 		FAttributeValidator AttributeValidator;
 		AttributeValidator.AddOptional(TEXT("axes"), ENNEAttributeDataType::Int32Array);
@@ -158,14 +160,25 @@ namespace UE::NNERuntimeRDG::Private::Hlsl
 		return bIsValid;
 	}
 
-	FOperatorHlsl* CreateReduceSumOperator()
+	template< UE::NNEHlslShaders::Internal::EReduceOperatorType ReduceOperatorType >
+	FOperatorHlsl* CreateReduceOperator()
 	{
-		return new FReduceSum();
+		return new FReduceOperator<ReduceOperatorType>();
 	}
 
 	bool RegisterReduceOperators(FOperatorRegistryHlsl& Registry)
 	{
-		Registry.OpAdd({{TEXT("ReduceSum"), TEXT("Onnx")}}, CreateReduceSumOperator, ValidateReduceSumOperator);
+		Registry.OpAdd({ {TEXT("ReduceL1"),       TEXT("Onnx")}}, CreateReduceOperator<UE::NNEHlslShaders::Internal::EReduceOperatorType::L1>, ValidateReduceOperator);
+		Registry.OpAdd({ {TEXT("ReduceL2"),       TEXT("Onnx")}}, CreateReduceOperator<UE::NNEHlslShaders::Internal::EReduceOperatorType::L2>, ValidateReduceOperator);
+		//ReduceLogSum not yet supported as multi axis case require to apply all reduction first then log.
+		Registry.OpAdd({{TEXT("ReduceLogSumExp"), TEXT("Onnx")}}, CreateReduceOperator<UE::NNEHlslShaders::Internal::EReduceOperatorType::LogSumExp>, ValidateReduceOperator);
+		Registry.OpAdd({{TEXT("ReduceMax"),       TEXT("Onnx")}}, CreateReduceOperator<UE::NNEHlslShaders::Internal::EReduceOperatorType::Max>, ValidateReduceOperator);
+		Registry.OpAdd({{TEXT("ReduceMean"),      TEXT("Onnx")}}, CreateReduceOperator<UE::NNEHlslShaders::Internal::EReduceOperatorType::Average>, ValidateReduceOperator);
+		Registry.OpAdd({{TEXT("ReduceMin"),       TEXT("Onnx")}}, CreateReduceOperator<UE::NNEHlslShaders::Internal::EReduceOperatorType::Min>, ValidateReduceOperator);
+		Registry.OpAdd({{TEXT("ReduceProd"),      TEXT("Onnx")}}, CreateReduceOperator<UE::NNEHlslShaders::Internal::EReduceOperatorType::Prod>, ValidateReduceOperator);
+		Registry.OpAdd({{TEXT("ReduceSum"),       TEXT("Onnx")}}, CreateReduceOperator<UE::NNEHlslShaders::Internal::EReduceOperatorType::Sum>, ValidateReduceOperator);
+		//ReduceSumSquare not yet supported as multi axis case require to apply square on whole tensor first then sum.
+
 		return true;
 	}
 } // UE::NNERuntimeRDG::Private::Hlsl
