[NNE] Gather operator implementation for NNERuntimeRDGHlsl

Tests:
* P4V Non-Unity build
* UnrealEditor, Win64: Gather tests, smoke tests
* UnrealGame, Win64: all tests

#rnx

#virtualized

diff --git a/Engine/Plugins/Experimental/NNERuntimeRDG/Source/NNEHlslShaders/Shaders/NNEHlslShadersGather.usf b/Engine/Plugins/Experimental/NNERuntimeRDG/Source/NNEHlslShaders/Shaders/NNEHlslShadersGather.usf
--- a/Engine/Plugins/Experimental/NNERuntimeRDG/Source/NNEHlslShaders/Shaders/NNEHlslShadersGather.usf
+++ b/Engine/Plugins/Experimental/NNERuntimeRDG/Source/NNEHlslShaders/Shaders/NNEHlslShadersGather.usf
@@ -50,7 +50,6 @@ void Gather(in const uint3 DispatchThreadID : SV_DispatchThreadID)
 	// Split the output indices into data and indices indices 
 	INDICES_WORK_TYPE DataIndices[NUM_OUTPUT_DIMENSIONS];
 	INDICES_WORK_TYPE IndicesIndices[NUM_OUTPUT_DIMENSIONS];
-	int AxisPlusNumIndicesDimensions = Axis + NumIndicesDimensions;
 	for (int i = 0; i < Axis; i++)
 	{
 		DataIndices[i] = OutputIndices[i];
@@ -61,7 +60,7 @@ void Gather(in const uint3 DispatchThreadID : SV_DispatchThreadID)
 	}
 	for (int i = Axis + 1; i < NumDataDimensions; i++)
 	{
-		DataIndices[i] = OutputIndices[AxisPlusNumIndicesDimensions + i];
+		DataIndices[i] = OutputIndices[NumIndicesDimensions + i - 1];
 	}
 
 	// Compute the flattened indices index
diff --git a/Engine/Plugins/Experimental/NNERuntimeRDG/Source/NNERuntimeRDG/Private/Hlsl/NNERuntimeRDGGather.cpp b/Engine/Plugins/Experimental/NNERuntimeRDG/Source/NNERuntimeRDG/Private/Hlsl/NNERuntimeRDGGather.cpp
--- a/Engine/Plugins/Experimental/NNERuntimeRDG/Source/NNERuntimeRDG/Private/Hlsl/NNERuntimeRDGGather.cpp
+++ b/Engine/Plugins/Experimental/NNERuntimeRDG/Source/NNERuntimeRDG/Private/Hlsl/NNERuntimeRDGGather.cpp
@@ -20,7 +20,7 @@ namespace UE::NNERuntimeRDG::Private::Hlsl
 	{
 	public:
 
-		FGather() {}
+		FGather() = default;
 		virtual ~FGather() = default;
 
 	private:
@@ -34,52 +34,47 @@ namespace UE::NNERuntimeRDG::Private::Hlsl
 			check(InputTensors.Num() == 2)
 			check(OutputTensors.Num() == 1)
 
-			const NNE::Internal::FTensor& Data = *InputTensors[0];
-			const NNE::Internal::FTensor& Indices = *InputTensors[1];
-			const NNE::FTensorShape& DataShape = Data.GetShape();
-			const NNE::FTensorShape& IndicesShape = Indices.GetShape();
+			const NNE::Internal::FTensor& DataTensor = *InputTensors[0];
+			const NNE::Internal::FTensor& IndicesTensor = *InputTensors[1];
+			const NNE::FTensorShape& DataShape = DataTensor.GetShape();
+			const NNE::FTensorShape& IndicesShape = IndicesTensor.GetShape();
 
 			const int32 OutputRank = IndicesShape.Rank() + DataShape.Rank() - 1;
+			
 			TArray<uint32> OutputShapeData;
-			int32 DataRankIdx = 0;
-
 			OutputShapeData.Reserve(OutputRank);
-			for (;DataRankIdx < Axis; ++DataRankIdx)
+
+			int32 DataRankIdx = 0;
+			while (DataRankIdx < Axis)
 			{
-				OutputShapeData.Add(DataShape.GetData()[DataRankIdx]);
+				OutputShapeData.Add(DataShape.GetData()[DataRankIdx++]);
 			}
 
 			OutputShapeData.Append(IndicesShape.GetData());
-			++DataRankIdx;
+			DataRankIdx++;
 
-			for (; DataRankIdx < DataShape.Rank(); ++DataRankIdx)
+			while (DataRankIdx < DataShape.Rank())
 			{
-				OutputShapeData.Add(DataShape.GetData()[DataRankIdx]);
+				OutputShapeData.Add(DataShape.GetData()[DataRankIdx++]);
 			}
 
 			NNE::FTensorShape OutputShape = NNE::FTensorShape::Make(OutputShapeData);
-
 			check(OutputShape.Rank() == OutputRank);
-			OutputTensors[0]->SetShape(OutputShape);
 
-			Internal::CPUHelper::Gather::Apply(Data, Indices, Axis, *OutputTensors[0]);
+			OutputTensors[0]->SetShape(OutputShape);
 
-			if (!OutputTensors[0]->HasPreparedData())
-			{
-				UE_LOG(LogNNE, Warning, TEXT("Gather: Output could not be computed as a constant tensor, however Gather is not implemented on GPU at the moment."));
-				return -1;
-			}
+			Internal::CPUHelper::Gather::Apply(DataTensor, IndicesTensor, Axis, *OutputTensors[0]);
 
 			return 0;
 		};
 		
 		virtual bool Initialize(TConstArrayView<NNE::FTensorDesc> InputTensorDescs, TConstArrayView<NNE::FTensorDesc> OutputTensorDescs, const NNE::FAttributeMap& Attributes) override
 		{
-			const int32 MaxNumDimensions = NNEHlslShaders::Internal::FGatherConstants::MAX_NUM_DIMENSIONS;
-			
 			check(InputTensorDescs.Num() == 2)
 			check(OutputTensorDescs.Num() == 1)
 
+			const int32 MaxNumDimensions = NNEHlslShaders::Internal::FGatherConstants::MAX_NUM_DIMENSIONS;
+
 			const NNE::FTensorDesc& Data = InputTensorDescs[0];
 			const NNE::FTensorDesc& Indices = InputTensorDescs[1];
 			const NNE::FTensorDesc& Output = OutputTensorDescs[0];
@@ -103,7 +98,7 @@ namespace UE::NNERuntimeRDG::Private::Hlsl
 				return false;
 			}
 			if (Axis < -Data.GetShape().Rank())
-{
+			{
 				UE_LOG(LogNNE, Warning, TEXT("Gather Axis attribute should be superior or equal to minus the first input rank"));
 				return false;
 			}
@@ -114,8 +109,6 @@ namespace UE::NNERuntimeRDG::Private::Hlsl
 
 		virtual void Dispatch(FRDGBuilder& GraphBuilder, TConstArrayView<FTensorRDGRef> InputTensors, TConstArrayView<FTensorRDGRef> OutputTensors) override
 		{
-			UE_LOG(LogNNE, Warning, TEXT("Gather: Output should be constant and already uploaded to GPU memory. Dispatch should not need to be called."));
-			/*
 			using namespace UE::NNEHlslShaders::Internal;
 
 			check(InputTensors.Num() == 2)
@@ -155,7 +148,6 @@ namespace UE::NNERuntimeRDG::Private::Hlsl
 				ComputeShader,
 				Parameters,
 				ThreadGroupCount);
-			*/
 		}
 	};
 
@@ -163,8 +155,8 @@ namespace UE::NNERuntimeRDG::Private::Hlsl
 	{
 		bool bIsValid = true;
 
-		//This match version 13 of the Gather operator
-		//[REDACTED_PROFILE_URL]/onnx/blob/main/docs/Operators.md#Gather
+		// This matches versions 1, 11 and 13 of the Gather operator
+		// [REDACTED_PROFILE_URL]/onnx/blob/main/docs/Operators.md#Gather
 		FAttributeValidator AttributeValidator;
 		AttributeValidator.AddOptional(TEXT("axis"), ENNEAttributeDataType::Int32);
 		bIsValid &= AttributeValidator.Validate(AttributeMap);
