[NNE] NNERuntimeHLSL GlobalAverage/MaxPool operators

Problem:
GlobalAverage/MaxPool were missing on the HLSL runtimes

Solution:
Implement them

Test: Win64 Editor: smoke tests + globalpool tests
Test: Win64 Standalone : all tests

#rnx

#virtualized

diff --git a/Engine/Plugins/Experimental/NNERuntimeRDG/Source/NNERuntimeRDG/Private/Hlsl/NNERuntimeRDGGlobalPool.cpp b/Engine/Plugins/Experimental/NNERuntimeRDG/Source/NNERuntimeRDG/Private/Hlsl/NNERuntimeRDGGlobalPool.cpp
new file mode 100644
--- /dev/null
+++ b/Engine/Plugins/Experimental/NNERuntimeRDG/Source/NNERuntimeRDG/Private/Hlsl/NNERuntimeRDGGlobalPool.cpp
@@ -0,0 +1,130 @@
+// Copyright Epic Games, Inc. All Rights Reserved.
+
+#include "NNERuntimeRDGGlobalPool.h"
+#include "NNEHlslShadersReduceCS.h"
+#include "NNERuntimeRDGHlslHelper.h"
+#include "NNETensor.h"
+#include "NNETypes.h"
+#include "RenderGraphUtils.h"
+
+namespace UE::NNERuntimeRDG::Private::Hlsl
+{
+	DECLARE_GPU_STAT_NAMED(FNNEOperatorGlobalPool, TEXT("NNE.Operator.Hlsl.GlobalPool"));
+
+	/**
+	 * GlobalPool operator implementation
+	 */
+	template< UE::NNEHlslShaders::Internal::EReduceOperatorType ReduceOperatorType >
+	class FGlobalPool : public FOperatorHlsl
+	{
+
+	public:
+
+		FGlobalPool() = default;
+		virtual ~FGlobalPool() = default;
+
+	private:
+
+		constexpr static int32 FirstReducedDimension = 2;
+
+	public:
+
+		virtual int PrepareOutputs(TConstArrayView<NNE::Internal::FTensorRef> InputTensors, TArrayView<NNE::Internal::FTensorRef> OutputTensors) const override
+		{
+			check(InputTensors.Num() == 1);
+			check(OutputTensors.Num() == 1);
+
+			const NNE::FTensorShape& InputShape = InputTensors[0]->GetShape();
+			const int32 InputRank = InputShape.Rank();
+
+			check(InputRank > FirstReducedDimension);
+
+			TArray<uint32, TInlineAllocator<NNE::FTensorShape::MaxRank>> OutputShape(InputShape.GetData());
+
+			for (int32 Axis = FirstReducedDimension; Axis < InputRank; ++Axis)
+			{
+				OutputShape[Axis] = 1;
+			}
+			OutputTensors[0]->SetShape(NNE::FTensorShape::Make(OutputShape));
+
+			return 0;
+		};
+
+		virtual bool Initialize(TConstArrayView<NNE::FTensorDesc> InputTensorDescs, TConstArrayView<NNE::FTensorDesc> OutputTensorDescs, const NNE::FAttributeMap& Attributes) override
+		{
+			check(InputTensorDescs.Num() == 1);
+			check(OutputTensorDescs.Num() == 1);
+
+			if (InputTensorDescs[0].GetShape().Rank() != OutputTensorDescs[0].GetShape().Rank())
+			{
+				UE_LOG(LogNNE, Warning, TEXT("GlobalPool operators requires the output to have the same rank as the input."));
+				return false;
+			}
+
+			const int32 InputRank = InputTensorDescs[0].GetShape().Rank();
+
+			if (InputRank <= FirstReducedDimension)
+			{
+				UE_LOG(LogNNE, Warning, TEXT("GlobalPool operators requires input tensor to be at least 3-D (but got rank %d)"), InputRank);
+				return false;
+			}
+
+			return true;
+		}
+
+		virtual void Dispatch(FRDGBuilder& GraphBuilder, TConstArrayView<FTensorRDGRef> InputTensors, TConstArrayView<FTensorRDGRef> OutputTensors) override
+		{
+			using namespace UE::NNEHlslShaders::Internal;
+
+			check(InputTensors.Num() == 1);
+			check(OutputTensors.Num() == 1);
+			check(InputTensors[0] != FTensorRDGRef{});
+			check(OutputTensors[0] != FTensorRDGRef{});
+
+			const FTensorRDG& Input = *InputTensors[0];
+			const FTensorRDG& Output = *OutputTensors[0];
+
+			RDG_EVENT_SCOPE(GraphBuilder, "NNE.Operator.Hlsl.GlobalPool");
+			RDG_GPU_STAT_SCOPE(GraphBuilder, FNNEOperatorGlobalPool);
+
+			FRDGBufferRef CurrInput = Input.GetBuffer();
+
+			TReduceCS::FParameters* Parameters = GraphBuilder.AllocParameters<TReduceCS::FParameters>();
+			TReduceCS::FillInParameters(Input.GetShape().GetData(), FirstReducedDimension, Parameters);
+			Parameters->AxisSize *= Parameters->NumElemAfterAxis;// GlobalPool reduce all trailing dimensions thus we can flatten them.
+			Parameters->NumElemAfterAxis = 1;
+			
+			TReduceCS::EnqueueRDG(GraphBuilder, Parameters, Input.GetBuffer(), Output.GetBuffer(), ReduceOperatorType);
+		}
+	};
+
+	bool ValidateGlobalPoolOperator(const NNE::FAttributeMap& AttributeMap, TConstArrayView<ENNETensorDataType> InputTypes, TConstArrayView<NNE::FSymbolicTensorShape> InputShapes)
+	{
+		bool bIsValid = true;
+
+		// This match version 1 of the GlobalAveragePool and GlobalMaxPool operators
+		// [REDACTED_PROFILE_URL]/onnx/blob/main/docs/Changelog.md#GlobalPool-1
+		FAttributeValidator AttributeValidator;
+		bIsValid &= AttributeValidator.Validate(AttributeMap);
+
+		FInputValidator InputValidator;
+		InputValidator.AddSupportedType(ENNETensorDataType::Float);
+		InputValidator.AddRequired();
+		bIsValid &= InputValidator.Validate(InputTypes);
+
+		return bIsValid;
+	}
+
+	template< UE::NNEHlslShaders::Internal::EReduceOperatorType ReduceOperatorType >
+	FOperatorHlsl* CreateGlobalPoolOperator()
+	{
+		return new FGlobalPool<ReduceOperatorType>();
+	}
+
+	bool RegisterGlobalPoolOperators(FOperatorRegistryHlsl& Registry)
+	{
+		Registry.OpAdd({{TEXT("GlobalAveragePool"), TEXT("Onnx")}}, CreateGlobalPoolOperator<UE::NNEHlslShaders::Internal::EReduceOperatorType::Average>, ValidateGlobalPoolOperator);
+		Registry.OpAdd({ {TEXT("GlobalMaxPool"), TEXT("Onnx")} }, CreateGlobalPoolOperator<UE::NNEHlslShaders::Internal::EReduceOperatorType::Max>, ValidateGlobalPoolOperator);
+		return true;
+	}
+} // UE::NNERuntimeRDG::Private::Hlsl
diff --git a/Engine/Plugins/Experimental/NNERuntimeRDG/Source/NNERuntimeRDG/Private/Hlsl/NNERuntimeRDGGlobalPool.h b/Engine/Plugins/Experimental/NNERuntimeRDG/Source/NNERuntimeRDG/Private/Hlsl/NNERuntimeRDGGlobalPool.h
new file mode 100644
--- /dev/null
+++ b/Engine/Plugins/Experimental/NNERuntimeRDG/Source/NNERuntimeRDG/Private/Hlsl/NNERuntimeRDGGlobalPool.h
@@ -0,0 +1,10 @@
+// Copyright Epic Games, Inc. All Rights Reserved.
+
+#pragma once
+
+#include "NNERuntimeRDGHlslOp.h"
+
+namespace UE::NNERuntimeRDG::Private::Hlsl
+{
+	bool RegisterGlobalPoolOperators(FOperatorRegistryHlsl& Registry);
+} // UE::NNERuntimeRDG::Private::Hlsl
diff --git a/Engine/Plugins/Experimental/NNERuntimeRDG/Source/NNERuntimeRDG/Private/NNERuntimeRDGHlsl.cpp b/Engine/Plugins/Experimental/NNERuntimeRDG/Source/NNERuntimeRDG/Private/NNERuntimeRDGHlsl.cpp
--- a/Engine/Plugins/Experimental/NNERuntimeRDG/Source/NNERuntimeRDG/Private/NNERuntimeRDGHlsl.cpp
+++ b/Engine/Plugins/Experimental/NNERuntimeRDG/Source/NNERuntimeRDG/Private/NNERuntimeRDGHlsl.cpp
@@ -26,6 +26,7 @@
 #include "Hlsl/NNERuntimeRDGFlatten.h"
 #include "Hlsl/NNERuntimeRDGGather.h"
 #include "Hlsl/NNERuntimeRDGGemm.h"
+#include "Hlsl/NNERuntimeRDGGlobalPool.h"
 #include "Hlsl/NNERuntimeRDGIdentity.h"
 #include "Hlsl/NNERuntimeRDGInstanceNormalization.h"
 #include "Hlsl/NNERuntimeRDGPad.h"
@@ -64,6 +65,7 @@ bool UNNERuntimeRDGHlslImpl::Init()
 	RegisterFlattenOperator(*Registry);
 	RegisterGatherOperator(*Registry);
 	RegisterGemmOperator(*Registry);
+	RegisterGlobalPoolOperators(*Registry);
 	RegisterIdentityOperator(*Registry);
 	RegisterInstanceNormalizationOperator(*Registry);
 	RegisterPadOperator(*Registry);
