[NNE] NNERuntimeHLSL fixes

Problems -> Solutions:
Conv default padding was wrong -> fixed
Gemm was missing support for C as 1x1 -> supported
NNETensor VS debugger custom previs was not working -> fix natvis

Bonus:
Cleanup following review feedback in GlobalPool hlsl op.
Enable now working inhouse model tests (opset 9)

Test: Win64 Editor: smoke + gemm + conv + inhouse model tests
Test: Win64 Standalone : all tests

#rnx

#virtualized

diff --git a/Engine/Plugins/Experimental/NNERuntimeRDG/Source/NNERuntimeRDG/Private/Hlsl/NNERuntimeRDGConv.cpp b/Engine/Plugins/Experimental/NNERuntimeRDG/Source/NNERuntimeRDG/Private/Hlsl/NNERuntimeRDGConv.cpp
--- a/Engine/Plugins/Experimental/NNERuntimeRDG/Source/NNERuntimeRDG/Private/Hlsl/NNERuntimeRDGConv.cpp
+++ b/Engine/Plugins/Experimental/NNERuntimeRDG/Source/NNERuntimeRDG/Private/Hlsl/NNERuntimeRDGConv.cpp
@@ -104,7 +104,7 @@ namespace UE::NNERuntimeRDG::Private::Hlsl
 			if (AutoPad == EConvAutoPad::NOTSET)
 			{
 				TArray<int32> PadsDefault;
-				PadsDefault.Init(1, 2 * NumDimensions);
+				PadsDefault.Init(0, 2 * NumDimensions);
 
 				Pads = Attributes.GetValueOrDefault<TArray<int32>>(TEXT("pads"), PadsDefault);
 			}
diff --git a/Engine/Plugins/Experimental/NNERuntimeRDG/Source/NNERuntimeRDG/Private/Hlsl/NNERuntimeRDGGemm.cpp b/Engine/Plugins/Experimental/NNERuntimeRDG/Source/NNERuntimeRDG/Private/Hlsl/NNERuntimeRDGGemm.cpp
--- a/Engine/Plugins/Experimental/NNERuntimeRDG/Source/NNERuntimeRDG/Private/Hlsl/NNERuntimeRDGGemm.cpp
+++ b/Engine/Plugins/Experimental/NNERuntimeRDG/Source/NNERuntimeRDG/Private/Hlsl/NNERuntimeRDGGemm.cpp
@@ -28,9 +28,6 @@ namespace UE::NNERuntimeRDG::Private::Hlsl
 		int32 InputTransA = 0;
 		int32 InputTransB = 0;
 
-		bool bIsCScalar = false;
-		bool bNoBias = true;
-
 	public:
 
 		virtual int PrepareOutputs(TConstArrayView<NNE::Internal::FTensorRef> InputTensors, TArrayView<NNE::Internal::FTensorRef> OutputTensors) const override
@@ -84,18 +81,8 @@ namespace UE::NNERuntimeRDG::Private::Hlsl
 					UE_LOG(LogNNE, Warning, TEXT("Gemm third input should be of rank 2 or less"));
 					return false;
 				}
-				if (InputC.GetShape().Rank() == 1 && InputC.GetShape().GetData()[0] == 1)
-				{
-					UE_LOG(LogNNE, Warning, TEXT("Gemm third input as scalar not supported"));
-					return false;
-				}
 			}
 
-			// C is treated as a scalar if there is no valid C, either width or height is zero or C dimension is 1x1
-			bIsCScalar = false; // InputTensors.Num() != 3 || InputC.Sizes[0] * InputC.Sizes[1] < 2;
-			// CScalar = C != nullptr ? C[0] : (InElementType)0;
-			bNoBias = InputTensorDescs.Num() != 3 /*|| InputC.Sizes[0] * InputC.Sizes[1] < 1*/;
-
 			InputAlpha = Attributes.GetValueOrDefault(TEXT("alpha"), InputAlpha);
 			InputBeta = Attributes.GetValueOrDefault(TEXT("beta"), InputBeta);
 			InputTransA = Attributes.GetValueOrDefault(TEXT("transA"), InputTransA);
@@ -108,29 +95,43 @@ namespace UE::NNERuntimeRDG::Private::Hlsl
 		{
 			using namespace UE::NNEHlslShaders::Internal;
 
-			const EGemmAlgorithm Algorithm = EGemmAlgorithm::Simple32x32;
-
-			const float CScalar = 0.0f;
-
 			check(InputTensors.Num() >= 2 && InputTensors.Num() <= 3);
 			check(OutputTensors.Num() == 1);
 			check(InputTensors[0] != nullptr);
 			check(InputTensors[1] != nullptr);
 			check(OutputTensors[0] != nullptr);
+
+			const EGemmAlgorithm Algorithm = EGemmAlgorithm::Simple32x32;
 			const FTensorRDG& InputA = *InputTensors[0];
 			const FTensorRDG& InputB = *InputTensors[1];
 			const FTensorRDG& Output = *OutputTensors[0];
 			const FTensorRDG* InputC = nullptr;
+			float CConstantScalar = 0.0f;
+			EGemmCScalar CScalarMode = EGemmCScalar::MAX;
 
 			if (InputTensors.Num() == 3)
 			{
 				check(InputTensors[2] != nullptr);
-				InputC = InputTensors[2];
+				if (InputTensors[2]->HasPreparedData() && InputTensors[2]->GetVolume() == 1)
+				{
+					CConstantScalar = InputTensors[2]->GetPreparedData<float>()[0];
+					CScalarMode = EGemmCScalar::Yes;
+				}
+				else
+				{
+					InputC = InputTensors[2];
+					CScalarMode = EGemmCScalar::No;
+				}
+			}
+			else
+			{
+				CScalarMode = EGemmCScalar::NoBias;
 			}
+			check(CScalarMode != EGemmCScalar::MAX);
 			
 			// Set parameters
 			TGemmCS::FParameters* Parameters = GraphBuilder.AllocParameters<TGemmCS::FParameters>();
-			TGemmCS::FillInParameters(InputAlpha, InputBeta, InputTransA, InputTransB, InputA, InputB, InputC, CScalar, *Parameters);
+			TGemmCS::FillInParameters(InputAlpha, InputBeta, InputTransA, InputTransB, InputA, InputB, InputC, CConstantScalar, *Parameters);
 			Parameters->A = GraphBuilder.CreateSRV(FRDGBufferSRVDesc(InputA.GetBuffer(), PF_R32_FLOAT));
 			Parameters->B = GraphBuilder.CreateSRV(FRDGBufferSRVDesc(InputB.GetBuffer(), PF_R32_FLOAT));
 			if (InputC != nullptr) {
@@ -139,7 +140,7 @@ namespace UE::NNERuntimeRDG::Private::Hlsl
 			Parameters->Y = GraphBuilder.CreateUAV(FRDGBufferUAVDesc(Output.GetBuffer(), PF_R32_FLOAT));
 
 			TGemmCS::FPermutationDomain PermutationVector;
-			PermutationVector.Set<TGemmCS::FGemmCScalar>(bNoBias ? EGemmCScalar::NoBias : (bIsCScalar ? EGemmCScalar::Yes : EGemmCScalar::No));
+			PermutationVector.Set<TGemmCS::FGemmCScalar>(CScalarMode);
 			PermutationVector.Set<TGemmCS::FGemmAlgorithm>(Algorithm);
 			PermutationVector.Set<TGemmCS::FGemmNumStackDimensions>(0);
 			TShaderMapRef<TGemmCS> ComputeShader(GetGlobalShaderMap(GMaxRHIFeatureLevel), PermutationVector);
@@ -163,6 +164,8 @@ namespace UE::NNERuntimeRDG::Private::Hlsl
 	{
 		bool bIsValid = true;
 
+		//This match version 9, 11, 13 of the Gemm operator see
+		//https://onnx.ai/onnx/operators/onnx__Gemm.html#l-onnx-doc-gemm
 		FAttributeValidator AttributeValidator;
 		AttributeValidator.AddOptional(TEXT("alpha"), ENNEAttributeDataType::Float);
 		AttributeValidator.AddOptional(TEXT("beta"), ENNEAttributeDataType::Float);
diff --git a/Engine/Plugins/Experimental/NNERuntimeRDG/Source/NNERuntimeRDG/Private/Hlsl/NNERuntimeRDGGlobalPool.cpp b/Engine/Plugins/Experimental/NNERuntimeRDG/Source/NNERuntimeRDG/Private/Hlsl/NNERuntimeRDGGlobalPool.cpp
--- a/Engine/Plugins/Experimental/NNERuntimeRDG/Source/NNERuntimeRDG/Private/Hlsl/NNERuntimeRDGGlobalPool.cpp
+++ b/Engine/Plugins/Experimental/NNERuntimeRDG/Source/NNERuntimeRDG/Private/Hlsl/NNERuntimeRDGGlobalPool.cpp
@@ -14,17 +14,17 @@ namespace UE::NNERuntimeRDG::Private::Hlsl
 	/**
 	 * GlobalPool operator implementation
 	 */
-	template< UE::NNEHlslShaders::Internal::EReduceOperatorType ReduceOperatorType >
 	class FGlobalPool : public FOperatorHlsl
 	{
 
 	public:
 
-		FGlobalPool() = default;
+		FGlobalPool(UE::NNEHlslShaders::Internal::EReduceOperatorType InReduceOperatorType):ReduceOperatorType(InReduceOperatorType) {};
 		virtual ~FGlobalPool() = default;
 
 	private:
 
+		const UE::NNEHlslShaders::Internal::EReduceOperatorType ReduceOperatorType;
 		constexpr static int32 FirstReducedDimension = 2;
 
 	public:
@@ -87,8 +87,6 @@ namespace UE::NNERuntimeRDG::Private::Hlsl
 			RDG_EVENT_SCOPE(GraphBuilder, "NNE.Operator.Hlsl.GlobalPool");
 			RDG_GPU_STAT_SCOPE(GraphBuilder, FNNEOperatorGlobalPool);
 
-			FRDGBufferRef CurrInput = Input.GetBuffer();
-
 			TReduceCS::FParameters* Parameters = GraphBuilder.AllocParameters<TReduceCS::FParameters>();
 			TReduceCS::FillInParameters(Input.GetShape().GetData(), FirstReducedDimension, Parameters);
 			Parameters->AxisSize *= Parameters->NumElemAfterAxis;// GlobalPool reduce all trailing dimensions thus we can flatten them.
@@ -118,7 +116,7 @@ namespace UE::NNERuntimeRDG::Private::Hlsl
 	template< UE::NNEHlslShaders::Internal::EReduceOperatorType ReduceOperatorType >
 	FOperatorHlsl* CreateGlobalPoolOperator()
 	{
-		return new FGlobalPool<ReduceOperatorType>();
+		return new FGlobalPool(ReduceOperatorType);
 	}
 
 	bool RegisterGlobalPoolOperators(FOperatorRegistryHlsl& Registry)
