[NNE] NNERuntimeHLSL Use factorized Reduction operation from Reduce for Softmax

Problem:
Softmax was doing inplace reduction rather than using the shared ReduceCS one, not benefiting for further potential optimization apply to it.

Solution:
Refactor the op to first apply the generic reduction then softmax in an element wise way.

Test: Win64 Editor: smoke tests + sofmax tests
Test: Win64 Standalone : all tests

#rnx

diff --git a/Engine/Plugins/Experimental/NNERuntimeRDG/Source/NNEHlslShaders/Internal/NNEHlslShadersSoftmaxCS.h b/Engine/Plugins/Experimental/NNERuntimeRDG/Source/NNEHlslShaders/Internal/NNEHlslShadersSoftmaxCS.h
--- a/Engine/Plugins/Experimental/NNERuntimeRDG/Source/NNEHlslShaders/Internal/NNEHlslShadersSoftmaxCS.h
+++ b/Engine/Plugins/Experimental/NNERuntimeRDG/Source/NNEHlslShaders/Internal/NNEHlslShadersSoftmaxCS.h
@@ -13,7 +13,7 @@ namespace UE::NNEHlslShaders::Internal
 	class FSoftmaxConstants
 	{
 	public:
-		static const int32 NUM_GROUP_THREADS{ 768 };
+		static const int32 NUM_GROUP_THREADS{ 256 };
 	};
 
 	class NNEHLSLSHADERS_API TSoftmaxCS : public FHlslShaderBase
@@ -24,13 +24,14 @@ namespace UE::NNEHlslShaders::Internal
 	public:
 
 		BEGIN_SHADER_PARAMETER_STRUCT(FParameters, )
-			SHADER_PARAMETER(int32, N)
-			SHADER_PARAMETER(int32, D)
+			SHADER_PARAMETER(uint32, Num)
+			SHADER_PARAMETER(uint32, ThreadCountX)
+			SHADER_PARAMETER(uint32, AxisSize)
 			SHADER_PARAMETER_RDG_BUFFER_SRV(Buffer<float>, Input)
+			SHADER_PARAMETER_RDG_BUFFER_SRV(Buffer<float>, InputSumExp)
 			SHADER_PARAMETER_RDG_BUFFER_UAV(RWBuffer<float>, Output)
 		END_SHADER_PARAMETER_STRUCT()
 
-		static FIntVector GetGroupCount(const FParameters& Parameters);
 		static void ModifyCompilationEnvironment(const FGlobalShaderPermutationParameters& InParameters, FShaderCompilerEnvironment& OutEnvironment);
 	};
 } // UE::NNEHlslShaders::Internal
\ No newline at end of file
diff --git a/Engine/Plugins/Experimental/NNERuntimeRDG/Source/NNEHlslShaders/Private/NNEHlslShadersSoftmaxCS.cpp b/Engine/Plugins/Experimental/NNERuntimeRDG/Source/NNEHlslShaders/Private/NNEHlslShadersSoftmaxCS.cpp
--- a/Engine/Plugins/Experimental/NNERuntimeRDG/Source/NNEHlslShaders/Private/NNEHlslShadersSoftmaxCS.cpp
+++ b/Engine/Plugins/Experimental/NNERuntimeRDG/Source/NNEHlslShaders/Private/NNEHlslShadersSoftmaxCS.cpp
@@ -5,11 +5,6 @@
 
 namespace UE::NNEHlslShaders::Internal
 {
-	FIntVector TSoftmaxCS::GetGroupCount(const TSoftmaxCS::FParameters& Parameters)
-	{
-		return { 1, (int32)Parameters.N, 1 };
-	}
-
 	void TSoftmaxCS::ModifyCompilationEnvironment(const FGlobalShaderPermutationParameters& InParameters, FShaderCompilerEnvironment& OutEnvironment)
 	{
 		FGlobalShader::ModifyCompilationEnvironment(InParameters, OutEnvironment);
diff --git a/Engine/Plugins/Experimental/NNERuntimeRDG/Source/NNEHlslShaders/Shaders/NNEHlslShadersSoftmax.usf b/Engine/Plugins/Experimental/NNERuntimeRDG/Source/NNEHlslShaders/Shaders/NNEHlslShadersSoftmax.usf
--- a/Engine/Plugins/Experimental/NNERuntimeRDG/Source/NNEHlslShaders/Shaders/NNEHlslShadersSoftmax.usf
+++ b/Engine/Plugins/Experimental/NNERuntimeRDG/Source/NNEHlslShaders/Shaders/NNEHlslShadersSoftmax.usf
@@ -8,65 +8,20 @@
 #define WRITE(x) x
 
 Buffer<WORK_TYPE> Input;
+Buffer<WORK_TYPE> InputSumExp;
 RWBuffer<WORK_TYPE> Output;
-int N;
-int D;
-
-groupshared WORK_TYPE SharedMemorySumExp[THREADGROUP_SIZE_X];
+uint Num;
+uint ThreadCountX;
+uint AxisSize;
 
 [numthreads(THREADGROUP_SIZE_X, 1, 1)]
-void Softmax(in const uint3 DispatchThreadID : SV_DispatchThreadID, in const uint3 GroupThreadID : SV_GroupThreadID)
+void Softmax(in const uint3 DispatchThreadID : SV_DispatchThreadID)
 {
-	const int StartIdx = DispatchThreadID.y * D;
-	const int SharedMemoryOffsetY = GroupThreadID.y * THREADGROUP_SIZE_X;
-
-	// Reduction using shared mem and THREADGROUP_SIZE_X thread
-	// https://en.wikipedia.org/wiki/Algorithms_for_calculating_variance
-	if (DispatchThreadID.y < N)
+	const uint Index = DispatchThreadID.y * ThreadCountX + DispatchThreadID.x;
+	if (Index < Num)
 	{
-		WORK_TYPE SumExp = 0;
-		int Count = 0;
-
-		for (int CurrIdx = GroupThreadID.x; CurrIdx < D; CurrIdx += THREADGROUP_SIZE_X)
-		{
-			Count += 1;
-
-			WORK_TYPE Value = READ(Input[StartIdx + CurrIdx]);
-
-			SumExp += exp(Value);
-		}
-
-		SharedMemorySumExp[SharedMemoryOffsetY + GroupThreadID.x] = SumExp;
-	}
-
-	// Wait until all threads wrote result to shared memory
-	GroupMemoryBarrierWithGroupSync();
-
-	// Use one thread to calculate final result
-	if (GroupThreadID.x == 0 && DispatchThreadID.y < N)
-	{
-		WORK_TYPE SumExp = 0;
-
-		UNROLL
-		for (int Idx = 0; Idx < THREADGROUP_SIZE_X; Idx++)
-		{
-			WORK_TYPE CurrSumExp = SharedMemorySumExp[SharedMemoryOffsetY + Idx];
-			SumExp += CurrSumExp;
-		}
-
-		SharedMemorySumExp[SharedMemoryOffsetY] = SumExp;
-	}
-
-	// Wait until first thread wrote result to shared memory
-	GroupMemoryBarrierWithGroupSync();
-
-	if (DispatchThreadID.y < N)
-	{
-		WORK_TYPE SumExp = SharedMemorySumExp[SharedMemoryOffsetY];
-
-		for (int CurrIdx = GroupThreadID.x; CurrIdx < D; CurrIdx += THREADGROUP_SIZE_X)
-		{
-			Output[StartIdx + CurrIdx] = WRITE(exp(READ(Input[StartIdx + CurrIdx])) / SumExp);
-		}
+		uint SumExpIndex = Index/AxisSize;
+		WORK_TYPE SumExp = READ(InputSumExp[SumExpIndex]);
+		Output[Index] = WRITE(exp(READ(Input[Index])) / SumExp);
 	}
 }
\ No newline at end of file
diff --git a/Engine/Plugins/Experimental/NNERuntimeRDG/Source/NNERuntimeRDG/Private/Hlsl/NNERuntimeRDGSoftmax.cpp b/Engine/Plugins/Experimental/NNERuntimeRDG/Source/NNERuntimeRDG/Private/Hlsl/NNERuntimeRDGSoftmax.cpp
--- a/Engine/Plugins/Experimental/NNERuntimeRDG/Source/NNERuntimeRDG/Private/Hlsl/NNERuntimeRDGSoftmax.cpp
+++ b/Engine/Plugins/Experimental/NNERuntimeRDG/Source/NNERuntimeRDG/Private/Hlsl/NNERuntimeRDGSoftmax.cpp
@@ -1,6 +1,7 @@
 // Copyright Epic Games, Inc. All Rights Reserved.
 
 #include "NNERuntimeRDGSoftmax.h"
+#include "NNEHlslShadersReduceCS.h"
 #include "NNEHlslShadersSoftmaxCS.h"
 #include "NNERuntimeRDGHlslHelper.h"
 #include "NNETensor.h"
@@ -87,40 +88,39 @@ namespace UE::NNERuntimeRDG::Private::Hlsl
 			check(Input.GetVolume() == Output.GetVolume());
 
 			const NNE::FTensorShape& InputShape = Input.GetShape();
-			TConstArrayView<uint32> InputShapeData = InputShape.GetData();
-			const int32 InputDimensions = InputShape.Rank();
-
-			int32 N = 1;
-			for (int32 i = 0; i < Axis; i++)
-			{
-				N *= InputShapeData[i];
-			}
-
-			int32 D = 1;
-			for (int32 i = Axis; i < InputDimensions; i++)
-			{
-				D *= InputShapeData[i];
-			}
-
-			// Set parameters
-			TSoftmaxCS::FParameters* Parameters = GraphBuilder.AllocParameters<TSoftmaxCS::FParameters>();
-			Parameters->N = N;
-			Parameters->D = D;
-			Parameters->Input = GraphBuilder.CreateSRV(FRDGBufferSRVDesc(Input.GetBuffer(), PF_R32_FLOAT));
-			Parameters->Output = GraphBuilder.CreateUAV(FRDGBufferUAVDesc(Output.GetBuffer(), PF_R32_FLOAT));
-
-			TShaderMapRef<TSoftmaxCS> ComputeShader(GetGlobalShaderMap(GMaxRHIFeatureLevel));
-			FIntVector ThreadGroupCount = TSoftmaxCS::GetGroupCount(*Parameters);
 
 			RDG_EVENT_SCOPE(GraphBuilder, "NNE.Operator.Hlsl.Softmax");
 			RDG_GPU_STAT_SCOPE(GraphBuilder, FNNEOperatorSoftmax);
 
+			// First apply Reduction(exp(x)) to temp buffer
+			TReduceCS::FParameters* ReduceParameters = GraphBuilder.AllocParameters<TReduceCS::FParameters>();
+			TReduceCS::FillInParameters(InputShape.GetData(), Axis, ReduceParameters);
+			ReduceParameters->AxisSize *= ReduceParameters->NumElemAfterAxis;// Softmax flatten the input tensor to a 2D one
+			ReduceParameters->NumElemAfterAxis = 1;
+			const FRDGBufferDesc SumExpBufferDesc = FRDGBufferDesc::CreateBufferDesc(Output.GetElementByteSize(), ReduceParameters->NumElemBeforeAxis);
+
+			FRDGBufferRef SumExpBuffer = GraphBuilder.CreateBuffer(SumExpBufferDesc, TEXT("NNE.Operator.Hlsl.Softmax.TempBuffer"), ERDGBufferFlags::None);
+
+			TReduceCS::EnqueueRDG(GraphBuilder, ReduceParameters, Input.GetBuffer(), SumExpBuffer, EReduceOperatorType::SumExp);
+
+			//Then Softmax
+			const int32 NumElements = Input.GetVolume();
+			const FIntVector ThreadGroupCount = ComputeElementWiseThreadGroups(NumElements, FSoftmaxConstants::NUM_GROUP_THREADS);
+			TSoftmaxCS::FParameters* SoftmaxParameters = GraphBuilder.AllocParameters<TSoftmaxCS::FParameters>();
+			SoftmaxParameters->AxisSize = ReduceParameters->AxisSize;
+			SoftmaxParameters->Num = NumElements;
+			SoftmaxParameters->ThreadCountX = ThreadGroupCount.X * FSoftmaxConstants::NUM_GROUP_THREADS;
+			SoftmaxParameters->Input = GraphBuilder.CreateSRV(FRDGBufferSRVDesc(Input.GetBuffer(), PF_R32_FLOAT));
+			SoftmaxParameters->InputSumExp = GraphBuilder.CreateSRV(FRDGBufferSRVDesc(SumExpBuffer, PF_R32_FLOAT));
+			SoftmaxParameters->Output = GraphBuilder.CreateUAV(FRDGBufferUAVDesc(Output.GetBuffer(), PF_R32_FLOAT));
+			TShaderMapRef<TSoftmaxCS> ComputeShader(GetGlobalShaderMap(GMaxRHIFeatureLevel));
+
 			FComputeShaderUtils::AddPass(
 				GraphBuilder,
 				RDG_EVENT_NAME("NNE.Operator.Hlsl.Softmax.Dispatch"),
 				ERDGPassFlags::Compute | ERDGPassFlags::NeverCull,
 				ComputeShader,
-				Parameters,
+				SoftmaxParameters,
 				ThreadGroupCount);
 		}
 	};
