[NNE] Turn unversioned RDG ops into versioned

Problem

RDG ops are all unversioned.

Solution

Put the right versions for each RDG op. The same implementation can be shared by multiple versions, given it's flexible.

Tests

- NNEIncubator Editor Win64
- NNEIncubator Standalone Win64
- Compile non-unity

#rnx

diff --git a/Engine/Plugins/Experimental/NNERuntimeRDG/Source/NNERuntimeRDG/Private/Dml/NNEDmlOperator.h b/Engine/Plugins/Experimental/NNERuntimeRDG/Source/NNERuntimeRDG/Private/Dml/NNEDmlOperator.h
--- a/Engine/Plugins/Experimental/NNERuntimeRDG/Source/NNERuntimeRDG/Private/Dml/NNEDmlOperator.h
+++ b/Engine/Plugins/Experimental/NNERuntimeRDG/Source/NNERuntimeRDG/Private/Dml/NNEDmlOperator.h
@@ -10,17 +10,42 @@
 #include "NNEAttributeMap.h"
 #include "NNERuntimeRDGBase.h"
 
+#define OP_DML_CLASS(OpName) FOperatorDml##OpName
+#define OP_DML_CLASS_TEMPL(OpName,OpVer) OP_TEST_CLASS(OpName) <OpVer>
+
 #define NNE_DML_REGISTER_OP(OpName) \
 struct FDmlOperator##OpName##Registrator \
 { \
 	FDmlOperator##OpName##Registrator() \
 	{ \
-		FOperatorRegistryDml::Get()->OpAdd({{TEXT(#OpName), TEXT("Onnx")}}, FOperatorDml##OpName##::Create, FOperatorDml##OpName##::Validate); \
+		FOperatorRegistryDml::Get()->OpAdd({{TEXT(#OpName), TEXT("Onnx")}}, OP_DML_CLASS(OpName)::Create, OP_DML_CLASS(OpName)::Validate); \
 	} \
 }; \
 \
 static FDmlOperator##OpName##Registrator RegisterDmlOperator##OpName;
 
+#define NNE_DML_REGISTER_OP_VERSION(OpName, OpVer) \
+struct FDmlOperator##OpName##OpVer##Registrator \
+{ \
+	FDmlOperator##OpName##OpVer##Registrator() \
+	{ \
+		FOperatorRegistryDml::Get()->OpAdd({{TEXT(#OpName), TEXT("Onnx")}, OpVer}, OP_DML_CLASS(OpName)::Create, OP_DML_CLASS(OpName)::Validate); \
+	} \
+}; \
+\
+static FDmlOperator##OpName##OpVer##Registrator RegisterDmlOperator##OpName##OpVer;
+
+#define NNE_DML_REGISTER_OP_TEMPL_VERSION(OpName, OpTemplVer, OpVer) \
+struct FDmlOperator##OpName##OpVer##Registrator \
+{ \
+	FDmlOperator##OpName##OpVer##Registrator() \
+	{ \
+		FOperatorRegistryDml::Get()->OpAdd({{TEXT(#OpName), TEXT("Onnx")}, OpVer}, OP_DML_CLASS_TEMPL(OpName,OpTemplVer)::Create, OP_DML_CLASS_TEMPL(OpName,OpTemplVer)::Validate); \
+	} \
+}; \
+\
+static FDmlOperator##OpName##OpVer##Registrator RegisterDmlOperator##OpName##OpVer;
+
 
 namespace UE::NNERuntimeRDG::Private::Dml
 {
diff --git a/Engine/Plugins/Experimental/NNERuntimeRDG/Source/NNERuntimeRDG/Private/Dml/NNEDmlOperatorBatchNormalization.cpp b/Engine/Plugins/Experimental/NNERuntimeRDG/Source/NNERuntimeRDG/Private/Dml/NNEDmlOperatorBatchNormalization.cpp
--- a/Engine/Plugins/Experimental/NNERuntimeRDG/Source/NNERuntimeRDG/Private/Dml/NNEDmlOperatorBatchNormalization.cpp
+++ b/Engine/Plugins/Experimental/NNERuntimeRDG/Source/NNERuntimeRDG/Private/Dml/NNEDmlOperatorBatchNormalization.cpp
@@ -186,7 +186,7 @@ private:
 };
 
 // Register operator on Module startup
-NNE_DML_REGISTER_OP(BatchNormalization)
+NNE_DML_REGISTER_OP_VERSION(BatchNormalization, 9)
 
 } // namespace UE::NNERuntimeRDG::Private::Dml
 
diff --git a/Engine/Plugins/Experimental/NNERuntimeRDG/Source/NNERuntimeRDG/Private/Dml/NNEDmlOperatorCast.cpp b/Engine/Plugins/Experimental/NNERuntimeRDG/Source/NNERuntimeRDG/Private/Dml/NNEDmlOperatorCast.cpp
--- a/Engine/Plugins/Experimental/NNERuntimeRDG/Source/NNERuntimeRDG/Private/Dml/NNEDmlOperatorCast.cpp
+++ b/Engine/Plugins/Experimental/NNERuntimeRDG/Source/NNERuntimeRDG/Private/Dml/NNEDmlOperatorCast.cpp
@@ -118,7 +118,9 @@ public:
 };
 
 // Register operator on Module startup
-NNE_DML_REGISTER_OP(Cast)
+NNE_DML_REGISTER_OP_VERSION(Cast, 6)
+NNE_DML_REGISTER_OP_VERSION(Cast, 9)
+NNE_DML_REGISTER_OP_VERSION(Cast, 13)
 
 } // namespace UE::NNERuntimeRDG::Private::Dml
 
diff --git a/Engine/Plugins/Experimental/NNERuntimeRDG/Source/NNERuntimeRDG/Private/Dml/NNEDmlOperatorConcat.cpp b/Engine/Plugins/Experimental/NNERuntimeRDG/Source/NNERuntimeRDG/Private/Dml/NNEDmlOperatorConcat.cpp
--- a/Engine/Plugins/Experimental/NNERuntimeRDG/Source/NNERuntimeRDG/Private/Dml/NNEDmlOperatorConcat.cpp
+++ b/Engine/Plugins/Experimental/NNERuntimeRDG/Source/NNERuntimeRDG/Private/Dml/NNEDmlOperatorConcat.cpp
@@ -177,7 +177,10 @@ public:
 };
 
 // Register Concat operator on Module startup
-NNE_DML_REGISTER_OP(Concat)
+NNE_DML_REGISTER_OP_VERSION(Concat, 1)
+NNE_DML_REGISTER_OP_VERSION(Concat, 4)
+NNE_DML_REGISTER_OP_VERSION(Concat, 11)
+NNE_DML_REGISTER_OP_VERSION(Concat, 13)
 
 } // namespace UE::NNERuntimeRDG::Private::Dml
 
diff --git a/Engine/Plugins/Experimental/NNERuntimeRDG/Source/NNERuntimeRDG/Private/Dml/NNEDmlOperatorConvolution.cpp b/Engine/Plugins/Experimental/NNERuntimeRDG/Source/NNERuntimeRDG/Private/Dml/NNEDmlOperatorConvolution.cpp
--- a/Engine/Plugins/Experimental/NNERuntimeRDG/Source/NNERuntimeRDG/Private/Dml/NNEDmlOperatorConvolution.cpp
+++ b/Engine/Plugins/Experimental/NNERuntimeRDG/Source/NNERuntimeRDG/Private/Dml/NNEDmlOperatorConvolution.cpp
@@ -258,12 +258,14 @@ public:
 
 void RegisterConvOperator()
 {
-	FOperatorRegistryDml::Get()->OpAdd({{TEXT("Conv"), TEXT("Onnx")}}, FOperatorDmlConv<DML_CONVOLUTION_DIRECTION_FORWARD>::Create);
+	FOperatorRegistryDml::Get()->OpAdd({{TEXT("Conv"), TEXT("Onnx")}, 1}, FOperatorDmlConv<DML_CONVOLUTION_DIRECTION_FORWARD>::Create);
+	FOperatorRegistryDml::Get()->OpAdd({{TEXT("Conv"), TEXT("Onnx")}, 11}, FOperatorDmlConv<DML_CONVOLUTION_DIRECTION_FORWARD>::Create);
 }
 
 void RegisterConvTransposeOperator()
 {
-	FOperatorRegistryDml::Get()->OpAdd({{TEXT("ConvTranspose"), TEXT("Onnx")}}, FOperatorDmlConv<DML_CONVOLUTION_DIRECTION_BACKWARD>::Create);
+	FOperatorRegistryDml::Get()->OpAdd({{TEXT("ConvTranspose"), TEXT("Onnx")}, 1}, FOperatorDmlConv<DML_CONVOLUTION_DIRECTION_BACKWARD>::Create);
+	FOperatorRegistryDml::Get()->OpAdd({{TEXT("ConvTranspose"), TEXT("Onnx")}, 11}, FOperatorDmlConv<DML_CONVOLUTION_DIRECTION_BACKWARD>::Create);
 }
 
 struct FDmlOperatorConvRegistrator
diff --git a/Engine/Plugins/Experimental/NNERuntimeRDG/Source/NNERuntimeRDG/Private/Dml/NNEDmlOperatorCumSum.cpp b/Engine/Plugins/Experimental/NNERuntimeRDG/Source/NNERuntimeRDG/Private/Dml/NNEDmlOperatorCumSum.cpp
--- a/Engine/Plugins/Experimental/NNERuntimeRDG/Source/NNERuntimeRDG/Private/Dml/NNEDmlOperatorCumSum.cpp
+++ b/Engine/Plugins/Experimental/NNERuntimeRDG/Source/NNERuntimeRDG/Private/Dml/NNEDmlOperatorCumSum.cpp
@@ -151,7 +151,8 @@ public:
 };
 
 // Register CumSum operator on Module startup
-NNE_DML_REGISTER_OP(CumSum)
+NNE_DML_REGISTER_OP_VERSION(CumSum, 11)
+NNE_DML_REGISTER_OP_VERSION(CumSum, 14)
 
 } // namespace UE::NNERuntimeRDG::Private::Dml
 
diff --git a/Engine/Plugins/Experimental/NNERuntimeRDG/Source/NNERuntimeRDG/Private/Dml/NNEDmlOperatorDepthToSpace.cpp b/Engine/Plugins/Experimental/NNERuntimeRDG/Source/NNERuntimeRDG/Private/Dml/NNEDmlOperatorDepthToSpace.cpp
--- a/Engine/Plugins/Experimental/NNERuntimeRDG/Source/NNERuntimeRDG/Private/Dml/NNEDmlOperatorDepthToSpace.cpp
+++ b/Engine/Plugins/Experimental/NNERuntimeRDG/Source/NNERuntimeRDG/Private/Dml/NNEDmlOperatorDepthToSpace.cpp
@@ -141,7 +141,9 @@ public:
 };
 
 // Register DepthToSpace operator on Module startup
-NNE_DML_REGISTER_OP(DepthToSpace)
+NNE_DML_REGISTER_OP_VERSION(DepthToSpace, 1)
+NNE_DML_REGISTER_OP_VERSION(DepthToSpace, 11)
+NNE_DML_REGISTER_OP_VERSION(DepthToSpace, 13)
 
 } // namespace UE::NNERuntimeRDG::Private::Dml
 
diff --git a/Engine/Plugins/Experimental/NNERuntimeRDG/Source/NNERuntimeRDG/Private/Dml/NNEDmlOperatorElementwise.cpp b/Engine/Plugins/Experimental/NNERuntimeRDG/Source/NNERuntimeRDG/Private/Dml/NNEDmlOperatorElementwise.cpp
--- a/Engine/Plugins/Experimental/NNERuntimeRDG/Source/NNERuntimeRDG/Private/Dml/NNEDmlOperatorElementwise.cpp
+++ b/Engine/Plugins/Experimental/NNERuntimeRDG/Source/NNERuntimeRDG/Private/Dml/NNEDmlOperatorElementwise.cpp
@@ -287,55 +287,73 @@ private:
 	}
 };
 
-#define REGISTER_OP_ELEMENT_WISE(OpName, DmlOpName, OpClass) \
-TCHAR const Op##OpName##Name[] = TEXT(#OpName); \
-struct FDmlOperator##OpName##Registrator \
+#define REGISTER_OP_ELEMENT_WISE(OpName, DmlOpName, OpClass, OpVer) \
+TCHAR const Op##OpName##OpVer##Name[] = TEXT(#OpName); \
+struct FDmlOperator##OpName##OpVer##Registrator \
 { \
-	FDmlOperator##OpName##Registrator() \
+	FDmlOperator##OpName##OpVer##Registrator() \
 	{ \
-		FOperatorRegistryDml::Get()->OpAdd({{TEXT(#OpName), TEXT("Onnx")}}, OpClass<DML_ELEMENT_WISE_##DmlOpName##_OPERATOR_DESC, DML_OPERATOR_ELEMENT_WISE_##DmlOpName, Op##OpName##Name>::Create, OpClass<DML_ELEMENT_WISE_##DmlOpName##_OPERATOR_DESC, DML_OPERATOR_ELEMENT_WISE_##DmlOpName, Op##OpName##Name>::Validate); \
+		FOperatorRegistryDml::Get()->OpAdd({{TEXT(#OpName), TEXT("Onnx")}, OpVer}, OpClass<DML_ELEMENT_WISE_##DmlOpName##_OPERATOR_DESC, DML_OPERATOR_ELEMENT_WISE_##DmlOpName, Op##OpName##OpVer##Name>::Create, OpClass<DML_ELEMENT_WISE_##DmlOpName##_OPERATOR_DESC, DML_OPERATOR_ELEMENT_WISE_##DmlOpName, Op##OpName##OpVer##Name>::Validate); \
 	} \
 }; \
 \
-static FDmlOperator##OpName##Registrator RegisterDmlOperator##OpName;
-
-
-#define REGISTER_OP_ELEMENT_WISE_UNARY(OpName, DmlOpName) REGISTER_OP_ELEMENT_WISE(OpName, DmlOpName, FOperatorDmlElementWiseUnary)
-
-#define REGISTER_OP_ELEMENT_WISE_BINARY(OpName, DmlOpName) REGISTER_OP_ELEMENT_WISE(OpName, DmlOpName, FOperatorDmlElementWiseBinary)
-
-REGISTER_OP_ELEMENT_WISE_UNARY( Abs,				ABS )
-REGISTER_OP_ELEMENT_WISE_UNARY( Acos, 				ACOS )
-REGISTER_OP_ELEMENT_WISE_UNARY( Acosh, 				ACOSH )
-REGISTER_OP_ELEMENT_WISE_UNARY( Asin, 				ASIN )
-REGISTER_OP_ELEMENT_WISE_UNARY( Asinh, 				ASINH )
-REGISTER_OP_ELEMENT_WISE_UNARY( Atan,				ATAN )
-REGISTER_OP_ELEMENT_WISE_UNARY( Atanh, 				ATANH )
-REGISTER_OP_ELEMENT_WISE_UNARY( Ceil,				CEIL )
-REGISTER_OP_ELEMENT_WISE_UNARY( Clip, 				CLIP )
-REGISTER_OP_ELEMENT_WISE_UNARY( Cos,				COS )
-REGISTER_OP_ELEMENT_WISE_UNARY( Cosh,				COSH )
-REGISTER_OP_ELEMENT_WISE_UNARY( Erf,				ERF )
-REGISTER_OP_ELEMENT_WISE_UNARY( Exp, 				EXP )
-REGISTER_OP_ELEMENT_WISE_UNARY( Floor,				FLOOR )
-REGISTER_OP_ELEMENT_WISE_UNARY( IsInf, 				IS_INFINITY )
-REGISTER_OP_ELEMENT_WISE_UNARY( IsNan, 				IS_NAN )
-REGISTER_OP_ELEMENT_WISE_UNARY( Log,				LOG )
-REGISTER_OP_ELEMENT_WISE_UNARY( Neg,				NEGATE )
-REGISTER_OP_ELEMENT_WISE_UNARY( Reciprocal, 		RECIP )
-REGISTER_OP_ELEMENT_WISE_UNARY( Round,				ROUND )
-REGISTER_OP_ELEMENT_WISE_UNARY( Sign,				SIGN )
-REGISTER_OP_ELEMENT_WISE_UNARY( Sin,				SIN )
-REGISTER_OP_ELEMENT_WISE_UNARY( Sinh,				SINH )
-REGISTER_OP_ELEMENT_WISE_UNARY( Sqrt,				SQRT )
-REGISTER_OP_ELEMENT_WISE_UNARY( Tan,				TAN )
-REGISTER_OP_ELEMENT_WISE_UNARY( Tanh,				TANH )
-
-REGISTER_OP_ELEMENT_WISE_BINARY( Add,				ADD )
-REGISTER_OP_ELEMENT_WISE_BINARY( Div,				DIVIDE )
-REGISTER_OP_ELEMENT_WISE_BINARY( Mul,				MULTIPLY )
-REGISTER_OP_ELEMENT_WISE_BINARY( Pow,				POW )
-REGISTER_OP_ELEMENT_WISE_BINARY( Sub,				SUBTRACT )
+static FDmlOperator##OpName##OpVer##Registrator RegisterDmlOperator##OpName##OpVer;
+
+
+#define REGISTER_OP_ELEMENT_WISE_UNARY(OpName, DmlOpName, OpVer) REGISTER_OP_ELEMENT_WISE(OpName, DmlOpName, FOperatorDmlElementWiseUnary, OpVer)
+
+#define REGISTER_OP_ELEMENT_WISE_BINARY(OpName, DmlOpName, OpVer) REGISTER_OP_ELEMENT_WISE(OpName, DmlOpName, FOperatorDmlElementWiseBinary, OpVer)
+
+#define REGISTER_OP_ELEMENT_WISE_UNARY_COMMON_VERSIONS(OpName, DmlOpName) \
+REGISTER_OP_ELEMENT_WISE_UNARY(OpName, DmlOpName, 6) \
+REGISTER_OP_ELEMENT_WISE_UNARY(OpName, DmlOpName, 13)
+
+#define REGISTER_OP_ELEMENT_WISE_BINARY_COMMON_VERSIONS(OpName, DmlOpName) \
+REGISTER_OP_ELEMENT_WISE_BINARY(OpName, DmlOpName, 6) \
+REGISTER_OP_ELEMENT_WISE_BINARY(OpName, DmlOpName, 7) \
+REGISTER_OP_ELEMENT_WISE_BINARY(OpName, DmlOpName, 13) \
+REGISTER_OP_ELEMENT_WISE_BINARY(OpName, DmlOpName, 14)
+
+REGISTER_OP_ELEMENT_WISE_UNARY_COMMON_VERSIONS( Abs,				ABS )
+REGISTER_OP_ELEMENT_WISE_UNARY( 			 	Acos, 				ACOS,		7 )
+REGISTER_OP_ELEMENT_WISE_UNARY( 			 	Acosh, 				ACOSH,		9 )
+REGISTER_OP_ELEMENT_WISE_UNARY( 			 	Asin, 				ASIN,		7 )
+REGISTER_OP_ELEMENT_WISE_UNARY( 			 	Asinh, 				ASINH,		9 )
+REGISTER_OP_ELEMENT_WISE_UNARY( 			 	Atan,				ATAN,		7 )
+REGISTER_OP_ELEMENT_WISE_UNARY( 			 	Atanh, 				ATANH,		9 )
+REGISTER_OP_ELEMENT_WISE_UNARY_COMMON_VERSIONS( Ceil,				CEIL )
+REGISTER_OP_ELEMENT_WISE_UNARY( 			 	Clip, 				CLIP, 		6 )
+REGISTER_OP_ELEMENT_WISE_UNARY( 			 	Cos,				COS,		7 )
+REGISTER_OP_ELEMENT_WISE_UNARY( 			 	Cosh,				COSH,		9 )
+REGISTER_OP_ELEMENT_WISE_UNARY( 			 	Erf,				ERF,		9 )
+REGISTER_OP_ELEMENT_WISE_UNARY( 			 	Erf,				ERF,		13 )
+REGISTER_OP_ELEMENT_WISE_UNARY_COMMON_VERSIONS( Exp, 				EXP )
+REGISTER_OP_ELEMENT_WISE_UNARY_COMMON_VERSIONS( Floor,				FLOOR )
+REGISTER_OP_ELEMENT_WISE_UNARY( 			 	IsInf, 				IS_INFINITY,10 )
+REGISTER_OP_ELEMENT_WISE_UNARY( 			 	IsInf, 				IS_INFINITY,20 )
+REGISTER_OP_ELEMENT_WISE_UNARY( 			 	IsNan, 				IS_NAN,		9 )
+REGISTER_OP_ELEMENT_WISE_UNARY( 			 	IsNan, 				IS_NAN,		13 )
+REGISTER_OP_ELEMENT_WISE_UNARY( 			 	IsNan, 				IS_NAN,		20 )
+REGISTER_OP_ELEMENT_WISE_UNARY_COMMON_VERSIONS( Log,				LOG )
+REGISTER_OP_ELEMENT_WISE_UNARY_COMMON_VERSIONS( Neg,				NEGATE )
+REGISTER_OP_ELEMENT_WISE_UNARY_COMMON_VERSIONS( Reciprocal, 		RECIP )
+REGISTER_OP_ELEMENT_WISE_UNARY( 			 	Round,				ROUND,		11 )
+REGISTER_OP_ELEMENT_WISE_UNARY( 			 	Sign,				SIGN,		9 )
+REGISTER_OP_ELEMENT_WISE_UNARY( 			 	Sign,				SIGN,		13 )
+REGISTER_OP_ELEMENT_WISE_UNARY( 			 	Sin,				SIN,		7 )
+REGISTER_OP_ELEMENT_WISE_UNARY( 			 	Sinh,				SINH,		9 )
+REGISTER_OP_ELEMENT_WISE_UNARY_COMMON_VERSIONS( Sqrt,				SQRT )
+REGISTER_OP_ELEMENT_WISE_UNARY( 			 	Tan,				TAN,		7 )
+REGISTER_OP_ELEMENT_WISE_UNARY_COMMON_VERSIONS( Tanh,				TANH )
+
+REGISTER_OP_ELEMENT_WISE_BINARY_COMMON_VERSIONS( Add,				ADD )
+REGISTER_OP_ELEMENT_WISE_BINARY_COMMON_VERSIONS( Div,				DIVIDE )
+REGISTER_OP_ELEMENT_WISE_BINARY_COMMON_VERSIONS( Mul,				MULTIPLY )
+REGISTER_OP_ELEMENT_WISE_BINARY( 			  	 Pow,				POW,		7 )
+REGISTER_OP_ELEMENT_WISE_BINARY( 			  	 Pow,				POW,		12 )
+REGISTER_OP_ELEMENT_WISE_BINARY( 			  	 Pow,				POW,		13 )
+REGISTER_OP_ELEMENT_WISE_BINARY( 			  	 Pow,				POW,		15 )
+REGISTER_OP_ELEMENT_WISE_BINARY_COMMON_VERSIONS( Sub,				SUBTRACT )
 
 } // namespace UE::NNERuntimeRDG::Private::Dml
 
diff --git a/Engine/Plugins/Experimental/NNERuntimeRDG/Source/NNERuntimeRDG/Private/Dml/NNEDmlOperatorFlatten.cpp b/Engine/Plugins/Experimental/NNERuntimeRDG/Source/NNERuntimeRDG/Private/Dml/NNEDmlOperatorFlatten.cpp
--- a/Engine/Plugins/Experimental/NNERuntimeRDG/Source/NNERuntimeRDG/Private/Dml/NNEDmlOperatorFlatten.cpp
+++ b/Engine/Plugins/Experimental/NNERuntimeRDG/Source/NNERuntimeRDG/Private/Dml/NNEDmlOperatorFlatten.cpp
@@ -114,7 +114,10 @@ public:
 };
 
 // Register Reshape operator on Module startup
-NNE_DML_REGISTER_OP(Flatten)
+NNE_DML_REGISTER_OP_VERSION(Flatten, 1)
+NNE_DML_REGISTER_OP_VERSION(Flatten, 9)
+NNE_DML_REGISTER_OP_VERSION(Flatten, 11)
+NNE_DML_REGISTER_OP_VERSION(Flatten, 13)
 
 } // namespace UE::NNERuntimeRDG::Private::Dml
 
diff --git a/Engine/Plugins/Experimental/NNERuntimeRDG/Source/NNERuntimeRDG/Private/Dml/NNEDmlOperatorGather.cpp b/Engine/Plugins/Experimental/NNERuntimeRDG/Source/NNERuntimeRDG/Private/Dml/NNEDmlOperatorGather.cpp
--- a/Engine/Plugins/Experimental/NNERuntimeRDG/Source/NNERuntimeRDG/Private/Dml/NNEDmlOperatorGather.cpp
+++ b/Engine/Plugins/Experimental/NNERuntimeRDG/Source/NNERuntimeRDG/Private/Dml/NNEDmlOperatorGather.cpp
@@ -174,7 +174,9 @@ public:
 };
 
 // Register operator on Module startup
-NNE_DML_REGISTER_OP(Gather)
+NNE_DML_REGISTER_OP_VERSION(Gather, 1)
+NNE_DML_REGISTER_OP_VERSION(Gather, 11)
+NNE_DML_REGISTER_OP_VERSION(Gather, 13)
 
 } // namespace UE::NNERuntimeRDG::Private::Dml
 
diff --git a/Engine/Plugins/Experimental/NNERuntimeRDG/Source/NNERuntimeRDG/Private/Dml/NNEDmlOperatorGemm.cpp b/Engine/Plugins/Experimental/NNERuntimeRDG/Source/NNERuntimeRDG/Private/Dml/NNEDmlOperatorGemm.cpp
--- a/Engine/Plugins/Experimental/NNERuntimeRDG/Source/NNERuntimeRDG/Private/Dml/NNEDmlOperatorGemm.cpp
+++ b/Engine/Plugins/Experimental/NNERuntimeRDG/Source/NNERuntimeRDG/Private/Dml/NNEDmlOperatorGemm.cpp
@@ -189,7 +189,10 @@ public:
 	}
 };
 
-NNE_DML_REGISTER_OP(Gemm)
+NNE_DML_REGISTER_OP_VERSION(Gemm, 7)
+NNE_DML_REGISTER_OP_VERSION(Gemm, 9)
+NNE_DML_REGISTER_OP_VERSION(Gemm, 11)
+NNE_DML_REGISTER_OP_VERSION(Gemm, 13)
 
 } // namespace UE::NNERuntimeRDG::Private::Dml
 
diff --git a/Engine/Plugins/Experimental/NNERuntimeRDG/Source/NNERuntimeRDG/Private/Dml/NNEDmlOperatorIdentity.cpp b/Engine/Plugins/Experimental/NNERuntimeRDG/Source/NNERuntimeRDG/Private/Dml/NNEDmlOperatorIdentity.cpp
--- a/Engine/Plugins/Experimental/NNERuntimeRDG/Source/NNERuntimeRDG/Private/Dml/NNEDmlOperatorIdentity.cpp
+++ b/Engine/Plugins/Experimental/NNERuntimeRDG/Source/NNERuntimeRDG/Private/Dml/NNEDmlOperatorIdentity.cpp
@@ -94,7 +94,11 @@ public:
 };
 
 // Register Reshape operator on Module startup
-NNE_DML_REGISTER_OP(Identity)
+NNE_DML_REGISTER_OP_VERSION(Identity, 1)
+NNE_DML_REGISTER_OP_VERSION(Identity, 13)
+NNE_DML_REGISTER_OP_VERSION(Identity, 14)
+NNE_DML_REGISTER_OP_VERSION(Identity, 16)
+NNE_DML_REGISTER_OP_VERSION(Identity, 19)
 
 } // namespace UE::NNERuntimeRDG::Private::Dml
 
diff --git a/Engine/Plugins/Experimental/NNERuntimeRDG/Source/NNERuntimeRDG/Private/Dml/NNEDmlOperatorInstanceNormalization.cpp b/Engine/Plugins/Experimental/NNERuntimeRDG/Source/NNERuntimeRDG/Private/Dml/NNEDmlOperatorInstanceNormalization.cpp
--- a/Engine/Plugins/Experimental/NNERuntimeRDG/Source/NNERuntimeRDG/Private/Dml/NNEDmlOperatorInstanceNormalization.cpp
+++ b/Engine/Plugins/Experimental/NNERuntimeRDG/Source/NNERuntimeRDG/Private/Dml/NNEDmlOperatorInstanceNormalization.cpp
@@ -194,7 +194,7 @@ public:
 };
 
 // Register operator on Module startup
-NNE_DML_REGISTER_OP(InstanceNormalization)
+NNE_DML_REGISTER_OP_VERSION(InstanceNormalization, 6)
 
 } // namespace UE::NNERuntimeRDG::Private::Dml
 
diff --git a/Engine/Plugins/Experimental/NNERuntimeRDG/Source/NNERuntimeRDG/Private/Dml/NNEDmlOperatorLRN.cpp b/Engine/Plugins/Experimental/NNERuntimeRDG/Source/NNERuntimeRDG/Private/Dml/NNEDmlOperatorLRN.cpp
--- a/Engine/Plugins/Experimental/NNERuntimeRDG/Source/NNERuntimeRDG/Private/Dml/NNEDmlOperatorLRN.cpp
+++ b/Engine/Plugins/Experimental/NNERuntimeRDG/Source/NNERuntimeRDG/Private/Dml/NNEDmlOperatorLRN.cpp
@@ -107,7 +107,8 @@ public:
 };
 
 // Register operator on Module startup
-NNE_DML_REGISTER_OP(LRN)
+NNE_DML_REGISTER_OP_VERSION(LRN, 1)
+NNE_DML_REGISTER_OP_VERSION(LRN, 13)
 
 } // namespace UE::NNERuntimeRDG::Private::Dml
 
diff --git a/Engine/Plugins/Experimental/NNERuntimeRDG/Source/NNERuntimeRDG/Private/Dml/NNEDmlOperatorLpNormalization.cpp b/Engine/Plugins/Experimental/NNERuntimeRDG/Source/NNERuntimeRDG/Private/Dml/NNEDmlOperatorLpNormalization.cpp
--- a/Engine/Plugins/Experimental/NNERuntimeRDG/Source/NNERuntimeRDG/Private/Dml/NNEDmlOperatorLpNormalization.cpp
+++ b/Engine/Plugins/Experimental/NNERuntimeRDG/Source/NNERuntimeRDG/Private/Dml/NNEDmlOperatorLpNormalization.cpp
@@ -104,7 +104,7 @@ public:
 };
 
 // Register operator on Module startup
-NNE_DML_REGISTER_OP(LpNormalization)
+NNE_DML_REGISTER_OP_VERSION(LpNormalization, 1)
 
 } // namespace UE::NNERuntimeRDG::Private::Dml
 
diff --git a/Engine/Plugins/Experimental/NNERuntimeRDG/Source/NNERuntimeRDG/Private/Dml/NNEDmlOperatorMatMul.cpp b/Engine/Plugins/Experimental/NNERuntimeRDG/Source/NNERuntimeRDG/Private/Dml/NNEDmlOperatorMatMul.cpp
--- a/Engine/Plugins/Experimental/NNERuntimeRDG/Source/NNERuntimeRDG/Private/Dml/NNEDmlOperatorMatMul.cpp
+++ b/Engine/Plugins/Experimental/NNERuntimeRDG/Source/NNERuntimeRDG/Private/Dml/NNEDmlOperatorMatMul.cpp
@@ -321,7 +321,9 @@ public:
 	}
 };
 
-NNE_DML_REGISTER_OP(MatMul)
+NNE_DML_REGISTER_OP_VERSION(MatMul, 1)
+NNE_DML_REGISTER_OP_VERSION(MatMul, 9)
+NNE_DML_REGISTER_OP_VERSION(MatMul, 13)
 
 } // namespace UE::NNERuntimeRDG::Private::Dml
 
diff --git a/Engine/Plugins/Experimental/NNERuntimeRDG/Source/NNERuntimeRDG/Private/Dml/NNEDmlOperatorMaxUnpool.cpp b/Engine/Plugins/Experimental/NNERuntimeRDG/Source/NNERuntimeRDG/Private/Dml/NNEDmlOperatorMaxUnpool.cpp
--- a/Engine/Plugins/Experimental/NNERuntimeRDG/Source/NNERuntimeRDG/Private/Dml/NNEDmlOperatorMaxUnpool.cpp
+++ b/Engine/Plugins/Experimental/NNERuntimeRDG/Source/NNERuntimeRDG/Private/Dml/NNEDmlOperatorMaxUnpool.cpp
@@ -207,7 +207,8 @@ public:
 };
 
 // Register MaxUnpool operator on Module startup
-NNE_DML_REGISTER_OP(MaxUnpool)
+NNE_DML_REGISTER_OP_VERSION(MaxUnpool, 9)
+NNE_DML_REGISTER_OP_VERSION(MaxUnpool, 11)
 
 } // namespace UE::NNERuntimeRDG::Private::Dml
 
diff --git a/Engine/Plugins/Experimental/NNERuntimeRDG/Source/NNERuntimeRDG/Private/Dml/NNEDmlOperatorMeanVarianceNormalization.cpp b/Engine/Plugins/Experimental/NNERuntimeRDG/Source/NNERuntimeRDG/Private/Dml/NNEDmlOperatorMeanVarianceNormalization.cpp
--- a/Engine/Plugins/Experimental/NNERuntimeRDG/Source/NNERuntimeRDG/Private/Dml/NNEDmlOperatorMeanVarianceNormalization.cpp
+++ b/Engine/Plugins/Experimental/NNERuntimeRDG/Source/NNERuntimeRDG/Private/Dml/NNEDmlOperatorMeanVarianceNormalization.cpp
@@ -131,7 +131,8 @@ public:
 };
 
 // Register operator on Module startup
-NNE_DML_REGISTER_OP(MeanVarianceNormalization)
+NNE_DML_REGISTER_OP_VERSION(MeanVarianceNormalization, 9)
+NNE_DML_REGISTER_OP_VERSION(MeanVarianceNormalization, 13)
 
 } // namespace UE::NNERuntimeRDG::Private::Dml
 
diff --git a/Engine/Plugins/Experimental/NNERuntimeRDG/Source/NNERuntimeRDG/Private/Dml/NNEDmlOperatorPad.cpp b/Engine/Plugins/Experimental/NNERuntimeRDG/Source/NNERuntimeRDG/Private/Dml/NNEDmlOperatorPad.cpp
--- a/Engine/Plugins/Experimental/NNERuntimeRDG/Source/NNERuntimeRDG/Private/Dml/NNEDmlOperatorPad.cpp
+++ b/Engine/Plugins/Experimental/NNERuntimeRDG/Source/NNERuntimeRDG/Private/Dml/NNEDmlOperatorPad.cpp
@@ -194,7 +194,9 @@ public:
 };
 
 // Register Pad operator on Module startup
-NNE_DML_REGISTER_OP(Pad)
+NNE_DML_REGISTER_OP_VERSION(Pad, 2)
+NNE_DML_REGISTER_OP_VERSION(Pad, 11)
+NNE_DML_REGISTER_OP_VERSION(Pad, 13)
 
 } // namespace UE::NNERuntimeRDG::Private::Dml
 
diff --git a/Engine/Plugins/Experimental/NNERuntimeRDG/Source/NNERuntimeRDG/Private/Dml/NNEDmlOperatorPool.cpp b/Engine/Plugins/Experimental/NNERuntimeRDG/Source/NNERuntimeRDG/Private/Dml/NNEDmlOperatorPool.cpp
--- a/Engine/Plugins/Experimental/NNERuntimeRDG/Source/NNERuntimeRDG/Private/Dml/NNEDmlOperatorPool.cpp
+++ b/Engine/Plugins/Experimental/NNERuntimeRDG/Source/NNERuntimeRDG/Private/Dml/NNEDmlOperatorPool.cpp
@@ -298,25 +298,32 @@ public:
 	}
 };
 
-#define NNE_DML_REGISTER_POOLING_OP(OpName, DmlPrefix, UseGlobalPooling) \
-TCHAR const Op##OpName##Name[] = TEXT(#OpName); \
-struct FDmlOperator##OpName##Registrator \
+#define NNE_DML_REGISTER_POOLING_OP(OpName, DmlPrefix, UseGlobalPooling, Version) \
+TCHAR const Op##OpName##Version##Name[] = TEXT(#OpName); \
+struct FDmlOperator##OpName##Version##Registrator \
 { \
-	FDmlOperator##OpName##Registrator() \
+	FDmlOperator##OpName##Version##Registrator() \
 	{ \
-		FOperatorRegistryDml::Get()->OpAdd({{TEXT(#OpName), TEXT("Onnx")}}, FOperatorDmlPool<DML_##DmlPrefix##_POOLING_OPERATOR_DESC, DML_OPERATOR_##DmlPrefix##_POOLING, UseGlobalPooling, Op##OpName##Name>::Create, FOperatorDmlPool<DML_##DmlPrefix##_POOLING_OPERATOR_DESC, DML_OPERATOR_##DmlPrefix##_POOLING, UseGlobalPooling, Op##OpName##Name>::Validate); \
+		FOperatorRegistryDml::Get()->OpAdd({{TEXT(#OpName), TEXT("Onnx")}, Version}, FOperatorDmlPool<DML_##DmlPrefix##_POOLING_OPERATOR_DESC, DML_OPERATOR_##DmlPrefix##_POOLING, UseGlobalPooling, Op##OpName##Version##Name>::Create, FOperatorDmlPool<DML_##DmlPrefix##_POOLING_OPERATOR_DESC, DML_OPERATOR_##DmlPrefix##_POOLING, UseGlobalPooling, Op##OpName##Version##Name>::Validate); \
 	} \
 }; \
 \
-static FDmlOperator##OpName##Registrator RegisterDmlOperator##OpName;
+static FDmlOperator##OpName##Version##Registrator RegisterDmlOperator##OpName##Version;
 
 // Register pooling operator on Module startup
-NNE_DML_REGISTER_POOLING_OP(MaxPool, MAX, false)
-NNE_DML_REGISTER_POOLING_OP(GlobalMaxPool, MAX, true)
-NNE_DML_REGISTER_POOLING_OP(AveragePool, AVERAGE, false)
-NNE_DML_REGISTER_POOLING_OP(GlobalAveragePool, AVERAGE, true)
-NNE_DML_REGISTER_POOLING_OP(LpPool, LP, false)
-NNE_DML_REGISTER_POOLING_OP(GlobalLpPool, LP, true)
+NNE_DML_REGISTER_POOLING_OP(MaxPool, MAX, false, 1)
+NNE_DML_REGISTER_POOLING_OP(MaxPool, MAX, false, 8)
+NNE_DML_REGISTER_POOLING_OP(MaxPool, MAX, false, 11)
+NNE_DML_REGISTER_POOLING_OP(MaxPool, MAX, false, 12)
+NNE_DML_REGISTER_POOLING_OP(GlobalMaxPool, MAX, true, 1)
+NNE_DML_REGISTER_POOLING_OP(AveragePool, AVERAGE, false, 1)
+NNE_DML_REGISTER_POOLING_OP(AveragePool, AVERAGE, false, 7)
+NNE_DML_REGISTER_POOLING_OP(AveragePool, AVERAGE, false, 10)
+NNE_DML_REGISTER_POOLING_OP(AveragePool, AVERAGE, false, 11)
+NNE_DML_REGISTER_POOLING_OP(GlobalAveragePool, AVERAGE, true, 1)
+NNE_DML_REGISTER_POOLING_OP(LpPool, LP, false, 2)
+NNE_DML_REGISTER_POOLING_OP(LpPool, LP, false, 11)
+NNE_DML_REGISTER_POOLING_OP(GlobalLpPool, LP, true, 2)
 
 #undef NNE_DML_REGISTER_POOLING_OP
 
diff --git a/Engine/Plugins/Experimental/NNERuntimeRDG/Source/NNERuntimeRDG/Private/Dml/NNEDmlOperatorReduce.cpp b/Engine/Plugins/Experimental/NNERuntimeRDG/Source/NNERuntimeRDG/Private/Dml/NNEDmlOperatorReduce.cpp
--- a/Engine/Plugins/Experimental/NNERuntimeRDG/Source/NNERuntimeRDG/Private/Dml/NNEDmlOperatorReduce.cpp
+++ b/Engine/Plugins/Experimental/NNERuntimeRDG/Source/NNERuntimeRDG/Private/Dml/NNEDmlOperatorReduce.cpp
@@ -238,31 +238,40 @@ public:
 };
 
 
-#define NNE_DML_REGISTER_REDUCE_OP(OpName, ReduceFunc) \
-TCHAR const Op##OpName##Name[] = TEXT(#OpName); \
-struct FDmlOperator##OpName##Registrator \
+#define NNE_DML_REGISTER_REDUCE_OP(OpName, ReduceFunc, Version) \
+TCHAR const Op##OpName##Version##Name[] = TEXT(#OpName); \
+struct FDmlOperator##OpName##Version##Registrator \
 { \
-	FDmlOperator##OpName##Registrator() \
+	FDmlOperator##OpName##Version##Registrator() \
 	{ \
-		FOperatorRegistryDml::Get()->OpAdd({{TEXT(#OpName), TEXT("Onnx")}}, FOperatorDmlReduce<ReduceFunc, Op##OpName##Name>::Create, FOperatorDmlReduce<ReduceFunc, Op##OpName##Name>::Validate); \
+		FOperatorRegistryDml::Get()->OpAdd({{TEXT(#OpName), TEXT("Onnx")}, Version}, FOperatorDmlReduce<ReduceFunc, Op##OpName##Version##Name>::Create, FOperatorDmlReduce<ReduceFunc, Op##OpName##Version##Name>::Validate); \
 	} \
 }; \
 \
-static FDmlOperator##OpName##Registrator RegisterDmlOperator##OpName;
-
-NNE_DML_REGISTER_REDUCE_OP(ReduceL1,		DML_REDUCE_FUNCTION_L1)
-NNE_DML_REGISTER_REDUCE_OP(ReduceL2,		DML_REDUCE_FUNCTION_L2)
-NNE_DML_REGISTER_REDUCE_OP(ReduceLogSum,	DML_REDUCE_FUNCTION_LOG_SUM)
-NNE_DML_REGISTER_REDUCE_OP(ReduceLogSumExp,	DML_REDUCE_FUNCTION_LOG_SUM_EXP)
-NNE_DML_REGISTER_REDUCE_OP(ReduceMin,		DML_REDUCE_FUNCTION_MIN)
-NNE_DML_REGISTER_REDUCE_OP(ReduceMax,		DML_REDUCE_FUNCTION_MAX)
-NNE_DML_REGISTER_REDUCE_OP(ReduceMean,		DML_REDUCE_FUNCTION_AVERAGE)
-NNE_DML_REGISTER_REDUCE_OP(ReduceProd,		DML_REDUCE_FUNCTION_MULTIPLY)
-NNE_DML_REGISTER_REDUCE_OP(ReduceSum,		DML_REDUCE_FUNCTION_SUM)
-NNE_DML_REGISTER_REDUCE_OP(ReduceSumSquare,	DML_REDUCE_FUNCTION_SUM_SQUARE)
-NNE_DML_REGISTER_REDUCE_OP(ArgMax,			DML_REDUCE_FUNCTION_ARGMAX)
-NNE_DML_REGISTER_REDUCE_OP(ArgMin,			DML_REDUCE_FUNCTION_ARGMIN)
-
+static FDmlOperator##OpName##Version##Registrator RegisterDmlOperator##OpName##Version;
+
+#define NNE_DML_REGISTER_REDUCE_OP_COMMON_VERSIONS(OpName, ReduceFunc) \
+NNE_DML_REGISTER_REDUCE_OP(OpName, ReduceFunc, 1) \
+NNE_DML_REGISTER_REDUCE_OP(OpName, ReduceFunc, 11) \
+NNE_DML_REGISTER_REDUCE_OP(OpName, ReduceFunc, 13)
+
+NNE_DML_REGISTER_REDUCE_OP_COMMON_VERSIONS(ReduceL1,			DML_REDUCE_FUNCTION_L1)
+NNE_DML_REGISTER_REDUCE_OP_COMMON_VERSIONS(ReduceL2,			DML_REDUCE_FUNCTION_L2)
+NNE_DML_REGISTER_REDUCE_OP_COMMON_VERSIONS(ReduceLogSum,		DML_REDUCE_FUNCTION_LOG_SUM)
+NNE_DML_REGISTER_REDUCE_OP_COMMON_VERSIONS(ReduceLogSumExp,	DML_REDUCE_FUNCTION_LOG_SUM_EXP)
+NNE_DML_REGISTER_REDUCE_OP_COMMON_VERSIONS(ReduceMin,			DML_REDUCE_FUNCTION_MIN)
+NNE_DML_REGISTER_REDUCE_OP(ReduceMin,						DML_REDUCE_FUNCTION_MIN, 12)
+NNE_DML_REGISTER_REDUCE_OP_COMMON_VERSIONS(ReduceMax,			DML_REDUCE_FUNCTION_MAX)
+NNE_DML_REGISTER_REDUCE_OP(ReduceMax,						DML_REDUCE_FUNCTION_MAX, 12)
+NNE_DML_REGISTER_REDUCE_OP_COMMON_VERSIONS(ReduceMean,			DML_REDUCE_FUNCTION_AVERAGE)
+NNE_DML_REGISTER_REDUCE_OP_COMMON_VERSIONS(ReduceProd,			DML_REDUCE_FUNCTION_MULTIPLY)
+NNE_DML_REGISTER_REDUCE_OP(ReduceSum,						DML_REDUCE_FUNCTION_SUM, 1)
+NNE_DML_REGISTER_REDUCE_OP(ReduceSum,						DML_REDUCE_FUNCTION_SUM, 11)
+NNE_DML_REGISTER_REDUCE_OP_COMMON_VERSIONS(ReduceSumSquare,	DML_REDUCE_FUNCTION_SUM_SQUARE)
+NNE_DML_REGISTER_REDUCE_OP_COMMON_VERSIONS(ArgMax,				DML_REDUCE_FUNCTION_ARGMAX)
+NNE_DML_REGISTER_REDUCE_OP(ArgMax,							DML_REDUCE_FUNCTION_ARGMAX, 12)
+NNE_DML_REGISTER_REDUCE_OP_COMMON_VERSIONS(ArgMin,				DML_REDUCE_FUNCTION_ARGMIN)
+NNE_DML_REGISTER_REDUCE_OP(ArgMin,							DML_REDUCE_FUNCTION_ARGMIN, 12)
 
 #undef NNE_DML_REGISTER_REDUCE_OP
 
diff --git a/Engine/Plugins/Experimental/NNERuntimeRDG/Source/NNERuntimeRDG/Private/Dml/NNEDmlOperatorResample.cpp b/Engine/Plugins/Experimental/NNERuntimeRDG/Source/NNERuntimeRDG/Private/Dml/NNEDmlOperatorResample.cpp
--- a/Engine/Plugins/Experimental/NNERuntimeRDG/Source/NNERuntimeRDG/Private/Dml/NNEDmlOperatorResample.cpp
+++ b/Engine/Plugins/Experimental/NNERuntimeRDG/Source/NNERuntimeRDG/Private/Dml/NNEDmlOperatorResample.cpp
@@ -466,20 +466,22 @@ public:
 	}
 };
 
-#define NNE_DML_REGISTER_RESAMPLE_OP(OpName, IsResize) \
-struct FDmlOperator##OpName##Registrator \
+#define NNE_DML_REGISTER_RESAMPLE_OP(OpName, IsResize, Version) \
+struct FDmlOperator##OpName##Version##Registrator \
 { \
-	FDmlOperator##OpName##Registrator() \
+	FDmlOperator##OpName##Version##Registrator() \
 	{ \
-		FOperatorRegistryDml::Get()->OpAdd({{TEXT(#OpName), TEXT("Onnx")}}, FOperatorDmlResample<IsResize>::Create, FOperatorDmlResample<IsResize>::Validate); \
+		FOperatorRegistryDml::Get()->OpAdd({{TEXT(#OpName), TEXT("Onnx")}, Version}, FOperatorDmlResample<IsResize>::Create, FOperatorDmlResample<IsResize>::Validate); \
 	} \
 }; \
 \
-static FDmlOperator##OpName##Registrator RegisterDmlOperator##OpName;
+static FDmlOperator##OpName##Version##Registrator RegisterDmlOperator##OpName##Version;
 
 // Register operator on Module startup
-NNE_DML_REGISTER_RESAMPLE_OP(Upsample, false)
-NNE_DML_REGISTER_RESAMPLE_OP(Resize, true)
+NNE_DML_REGISTER_RESAMPLE_OP(Upsample, false, 9)
+NNE_DML_REGISTER_RESAMPLE_OP(Resize, true, 10)
+NNE_DML_REGISTER_RESAMPLE_OP(Resize, true, 11)
+NNE_DML_REGISTER_RESAMPLE_OP(Resize, true, 13)
 
 } // namespace UE::NNERuntimeRDG::Private::Dml
 
diff --git a/Engine/Plugins/Experimental/NNERuntimeRDG/Source/NNERuntimeRDG/Private/Dml/NNEDmlOperatorReshape.cpp b/Engine/Plugins/Experimental/NNERuntimeRDG/Source/NNERuntimeRDG/Private/Dml/NNEDmlOperatorReshape.cpp
--- a/Engine/Plugins/Experimental/NNERuntimeRDG/Source/NNERuntimeRDG/Private/Dml/NNEDmlOperatorReshape.cpp
+++ b/Engine/Plugins/Experimental/NNERuntimeRDG/Source/NNERuntimeRDG/Private/Dml/NNEDmlOperatorReshape.cpp
@@ -147,7 +147,10 @@ public:
 };
 
 // Register Reshape operator on Module startup
-NNE_DML_REGISTER_OP(Reshape)
+NNE_DML_REGISTER_OP_VERSION(Reshape, 5)
+NNE_DML_REGISTER_OP_VERSION(Reshape, 13)
+NNE_DML_REGISTER_OP_VERSION(Reshape, 14)
+NNE_DML_REGISTER_OP_VERSION(Reshape, 19)
 
 } // namespace UE::NNERuntimeRDG::Private::Dml
 
diff --git a/Engine/Plugins/Experimental/NNERuntimeRDG/Source/NNERuntimeRDG/Private/Dml/NNEDmlOperatorSlice.cpp b/Engine/Plugins/Experimental/NNERuntimeRDG/Source/NNERuntimeRDG/Private/Dml/NNEDmlOperatorSlice.cpp
--- a/Engine/Plugins/Experimental/NNERuntimeRDG/Source/NNERuntimeRDG/Private/Dml/NNEDmlOperatorSlice.cpp
+++ b/Engine/Plugins/Experimental/NNERuntimeRDG/Source/NNERuntimeRDG/Private/Dml/NNEDmlOperatorSlice.cpp
@@ -289,7 +289,9 @@ public:
 };
 
 // Register Slice operator on Module startup
-NNE_DML_REGISTER_OP(Slice)
+NNE_DML_REGISTER_OP_VERSION(Slice, 10)
+NNE_DML_REGISTER_OP_VERSION(Slice, 11)
+NNE_DML_REGISTER_OP_VERSION(Slice, 13)
 
 } // namespace UE::NNERuntimeRDG::Private::Dml
 
diff --git a/Engine/Plugins/Experimental/NNERuntimeRDG/Source/NNERuntimeRDG/Private/Dml/NNEDmlOperatorSpaceToDepth.cpp b/Engine/Plugins/Experimental/NNERuntimeRDG/Source/NNERuntimeRDG/Private/Dml/NNEDmlOperatorSpaceToDepth.cpp
--- a/Engine/Plugins/Experimental/NNERuntimeRDG/Source/NNERuntimeRDG/Private/Dml/NNEDmlOperatorSpaceToDepth.cpp
+++ b/Engine/Plugins/Experimental/NNERuntimeRDG/Source/NNERuntimeRDG/Private/Dml/NNEDmlOperatorSpaceToDepth.cpp
@@ -126,7 +126,8 @@ public:
 };
 
 // Register SpaceToDepth operator on Module startup
-NNE_DML_REGISTER_OP(SpaceToDepth)
+NNE_DML_REGISTER_OP_VERSION(SpaceToDepth, 1)
+NNE_DML_REGISTER_OP_VERSION(SpaceToDepth, 13)
 
 } // namespace UE::NNERuntimeRDG::Private::Dml
 
diff --git a/Engine/Plugins/Experimental/NNERuntimeRDG/Source/NNERuntimeRDG/Private/Dml/NNEDmlOperatorSplit.cpp b/Engine/Plugins/Experimental/NNERuntimeRDG/Source/NNERuntimeRDG/Private/Dml/NNEDmlOperatorSplit.cpp
--- a/Engine/Plugins/Experimental/NNERuntimeRDG/Source/NNERuntimeRDG/Private/Dml/NNEDmlOperatorSplit.cpp
+++ b/Engine/Plugins/Experimental/NNERuntimeRDG/Source/NNERuntimeRDG/Private/Dml/NNEDmlOperatorSplit.cpp
@@ -176,7 +176,8 @@ public:
 };
 
 // Register Split operator on Module startup
-NNE_DML_REGISTER_OP(Split)
+NNE_DML_REGISTER_OP_VERSION(Split, 2)
+NNE_DML_REGISTER_OP_VERSION(Split, 11)
 
 } // namespace UE::NNERuntimeRDG::Private::Dml
 
diff --git a/Engine/Plugins/Experimental/NNERuntimeRDG/Source/NNERuntimeRDG/Private/Dml/NNEDmlOperatorSqueeze.cpp b/Engine/Plugins/Experimental/NNERuntimeRDG/Source/NNERuntimeRDG/Private/Dml/NNEDmlOperatorSqueeze.cpp
--- a/Engine/Plugins/Experimental/NNERuntimeRDG/Source/NNERuntimeRDG/Private/Dml/NNEDmlOperatorSqueeze.cpp
+++ b/Engine/Plugins/Experimental/NNERuntimeRDG/Source/NNERuntimeRDG/Private/Dml/NNEDmlOperatorSqueeze.cpp
@@ -130,8 +130,9 @@ public:
 	}
 };
 
-// Register Reshape operator on Module startup
-NNE_DML_REGISTER_OP(Squeeze)
+// Register Squeeze operator on Module startup
+NNE_DML_REGISTER_OP_VERSION(Squeeze, 1)
+NNE_DML_REGISTER_OP_VERSION(Squeeze, 11)
 
 } // namespace UE::NNERuntimeRDG::Private::Dml
 
diff --git a/Engine/Plugins/Experimental/NNERuntimeRDG/Source/NNERuntimeRDG/Private/Dml/NNEDmlOperatorTranspose.cpp b/Engine/Plugins/Experimental/NNERuntimeRDG/Source/NNERuntimeRDG/Private/Dml/NNEDmlOperatorTranspose.cpp
--- a/Engine/Plugins/Experimental/NNERuntimeRDG/Source/NNERuntimeRDG/Private/Dml/NNEDmlOperatorTranspose.cpp
+++ b/Engine/Plugins/Experimental/NNERuntimeRDG/Source/NNERuntimeRDG/Private/Dml/NNEDmlOperatorTranspose.cpp
@@ -137,7 +137,8 @@ public:
 };
 
 // Register Transpose operator on Module startup
-NNE_DML_REGISTER_OP(Transpose)
+NNE_DML_REGISTER_OP_VERSION(Transpose, 1)
+NNE_DML_REGISTER_OP_VERSION(Transpose, 13)
 
 } // namespace UE::NNERuntimeRDG::Private::Dml
 
diff --git a/Engine/Plugins/Experimental/NNERuntimeRDG/Source/NNERuntimeRDG/Private/Dml/NNEDmlOperatorUnsqueeze.cpp b/Engine/Plugins/Experimental/NNERuntimeRDG/Source/NNERuntimeRDG/Private/Dml/NNEDmlOperatorUnsqueeze.cpp
--- a/Engine/Plugins/Experimental/NNERuntimeRDG/Source/NNERuntimeRDG/Private/Dml/NNEDmlOperatorUnsqueeze.cpp
+++ b/Engine/Plugins/Experimental/NNERuntimeRDG/Source/NNERuntimeRDG/Private/Dml/NNEDmlOperatorUnsqueeze.cpp
@@ -118,7 +118,8 @@ public:
 };
 
 // Register Reshape operator on Module startup
-NNE_DML_REGISTER_OP(Unsqueeze)
+NNE_DML_REGISTER_OP_VERSION(Unsqueeze, 1)
+NNE_DML_REGISTER_OP_VERSION(Unsqueeze, 11)
 
 } // namespace UE::NNERuntimeRDG::Private::Dml
 
diff --git a/Engine/Plugins/Experimental/NNERuntimeRDG/Source/NNERuntimeRDG/Private/Hlsl/NNERuntimeRDGBatchNormalization.cpp b/Engine/Plugins/Experimental/NNERuntimeRDG/Source/NNERuntimeRDG/Private/Hlsl/NNERuntimeRDGBatchNormalization.cpp
--- a/Engine/Plugins/Experimental/NNERuntimeRDG/Source/NNERuntimeRDG/Private/Hlsl/NNERuntimeRDGBatchNormalization.cpp
+++ b/Engine/Plugins/Experimental/NNERuntimeRDG/Source/NNERuntimeRDG/Private/Hlsl/NNERuntimeRDGBatchNormalization.cpp
@@ -126,8 +126,6 @@ namespace UE::NNERuntimeRDG::Private::Hlsl
 	{
 		bool bIsValid = true;
 
-		//This match version 9 of the BatchNormalization operator
-		//[REDACTED_PROFILE_URL]/onnx/blob/main/docs/Operators.md#BatchNormalization
 		FAttributeValidator AttributeValidator;
 		AttributeValidator.AddOptional(TEXT("epsilon"), ENNEAttributeDataType::Float);
 		AttributeValidator.AddOptional(TEXT("momentum"), ENNEAttributeDataType::Float); //Will be ignored, only useful in training mode.
@@ -154,7 +152,9 @@ namespace UE::NNERuntimeRDG::Private::Hlsl
 
 	bool RegisterBatchNormalizationOperator(FOperatorRegistryHlsl& Registry)
 	{
-		Registry.OpAdd({{TEXT("BatchNormalization"), TEXT("Onnx")}}, CreateBatchNormalizationOperator, ValidateBatchNormalizationOperator);
+		// Note: support of a particular version is partial with respect to tensor data types (only the most typical ones are usually supported).
+		Registry.OpAdd({{TEXT("BatchNormalization"), TEXT("Onnx")}, 9}, CreateBatchNormalizationOperator, ValidateBatchNormalizationOperator);
+		// Next version: 14
 		return true;
 	}
 } // UE::NNERuntimeRDG::Private::Hlsl
diff --git a/Engine/Plugins/Experimental/NNERuntimeRDG/Source/NNERuntimeRDG/Private/Hlsl/NNERuntimeRDGCast.cpp b/Engine/Plugins/Experimental/NNERuntimeRDG/Source/NNERuntimeRDG/Private/Hlsl/NNERuntimeRDGCast.cpp
--- a/Engine/Plugins/Experimental/NNERuntimeRDG/Source/NNERuntimeRDG/Private/Hlsl/NNERuntimeRDGCast.cpp
+++ b/Engine/Plugins/Experimental/NNERuntimeRDG/Source/NNERuntimeRDG/Private/Hlsl/NNERuntimeRDGCast.cpp
@@ -65,8 +65,6 @@ namespace UE::NNERuntimeRDG::Private::Hlsl
 	{
 		bool bIsValid = true;
 
-		//This match version 13 of the Cast operator
-		//[REDACTED_PROFILE_URL]/onnx/blob/main/docs/Operators.md#Cast
 		FAttributeValidator AttributeValidator;
 		AttributeValidator.AddRequired(TEXT("to"), ENNEAttributeDataType::Int32);
 		bIsValid &= AttributeValidator.Validate(AttributeMap);
@@ -106,7 +104,11 @@ namespace UE::NNERuntimeRDG::Private::Hlsl
 
 	bool RegisterCastOperator(FOperatorRegistryHlsl& Registry)
 	{
-		Registry.OpAdd({{TEXT("Cast"), TEXT("Onnx")}}, CreateCastOperator, ValidateCastOperator);
+		// Note: support of a particular version is partial with respect to tensor data types (only the most typical ones are usually supported).
+		Registry.OpAdd({{TEXT("Cast"), TEXT("Onnx")}, 6}, CreateCastOperator, ValidateCastOperator);
+		Registry.OpAdd({{TEXT("Cast"), TEXT("Onnx")}, 9}, CreateCastOperator, ValidateCastOperator);
+		Registry.OpAdd({{TEXT("Cast"), TEXT("Onnx")}, 13}, CreateCastOperator, ValidateCastOperator);
+		// Next version: 19
 		return true;
 	}
 } // UE::NNERuntimeRDG::Private::Hlsl
diff --git a/Engine/Plugins/Experimental/NNERuntimeRDG/Source/NNERuntimeRDG/Private/Hlsl/NNERuntimeRDGConcat.cpp b/Engine/Plugins/Experimental/NNERuntimeRDG/Source/NNERuntimeRDG/Private/Hlsl/NNERuntimeRDGConcat.cpp
--- a/Engine/Plugins/Experimental/NNERuntimeRDG/Source/NNERuntimeRDG/Private/Hlsl/NNERuntimeRDGConcat.cpp
+++ b/Engine/Plugins/Experimental/NNERuntimeRDG/Source/NNERuntimeRDG/Private/Hlsl/NNERuntimeRDGConcat.cpp
@@ -151,8 +151,6 @@ namespace UE::NNERuntimeRDG::Private::Hlsl
 	{
 		bool bIsValid = true;
 
-		//This match version 11 of the Concat operator
-		//[REDACTED_PROFILE_URL]/onnx/blob/main/docs/Operators.md#Concat
 		FAttributeValidator AttributeValidator;
 		AttributeValidator.AddRequired(TEXT("axis"), ENNEAttributeDataType::Int32);
 		bIsValid &= AttributeValidator.Validate(AttributeMap);
@@ -181,7 +179,10 @@ namespace UE::NNERuntimeRDG::Private::Hlsl
 
 	bool RegisterConcatOperator(FOperatorRegistryHlsl& Registry)
 	{
-		Registry.OpAdd({{TEXT("Concat"), TEXT("Onnx")}}, CreateConcatOperator, ValidateConcatOperator);
+		// Note: support of a particular version is partial with respect to tensor data types (only the most typical ones are usually supported).
+		Registry.OpAdd({{TEXT("Concat"), TEXT("Onnx")}, 4}, CreateConcatOperator, ValidateConcatOperator);
+		Registry.OpAdd({{TEXT("Concat"), TEXT("Onnx")}, 11}, CreateConcatOperator, ValidateConcatOperator);
+		Registry.OpAdd({{TEXT("Concat"), TEXT("Onnx")}, 13}, CreateConcatOperator, ValidateConcatOperator);
 		return true;
 	}
 } // UE::NNERuntimeRDG::Private::Hlsl
diff --git a/Engine/Plugins/Experimental/NNERuntimeRDG/Source/NNERuntimeRDG/Private/Hlsl/NNERuntimeRDGConv.cpp b/Engine/Plugins/Experimental/NNERuntimeRDG/Source/NNERuntimeRDG/Private/Hlsl/NNERuntimeRDGConv.cpp
--- a/Engine/Plugins/Experimental/NNERuntimeRDG/Source/NNERuntimeRDG/Private/Hlsl/NNERuntimeRDGConv.cpp
+++ b/Engine/Plugins/Experimental/NNERuntimeRDG/Source/NNERuntimeRDG/Private/Hlsl/NNERuntimeRDGConv.cpp
@@ -324,8 +324,9 @@ namespace UE::NNERuntimeRDG::Private::Hlsl
 
 	bool RegisterConvOperator(FOperatorRegistryHlsl& Registry)
 	{
-		Registry.OpAdd({{TEXT("Conv"), TEXT("Onnx")}}, FConv::Create, ValidateConvOperator);
-
+		// Note: support of a particular version is partial with respect to tensor data types (only the most typical ones are usually supported).
+		Registry.OpAdd({{TEXT("Conv"), TEXT("Onnx")}, 1}, FConv::Create, ValidateConvOperator);
+		Registry.OpAdd({{TEXT("Conv"), TEXT("Onnx")}, 11}, FConv::Create, ValidateConvOperator);
 		return true;
 	}
 
diff --git a/Engine/Plugins/Experimental/NNERuntimeRDG/Source/NNERuntimeRDG/Private/Hlsl/NNERuntimeRDGConvTranspose.cpp b/Engine/Plugins/Experimental/NNERuntimeRDG/Source/NNERuntimeRDG/Private/Hlsl/NNERuntimeRDGConvTranspose.cpp
--- a/Engine/Plugins/Experimental/NNERuntimeRDG/Source/NNERuntimeRDG/Private/Hlsl/NNERuntimeRDGConvTranspose.cpp
+++ b/Engine/Plugins/Experimental/NNERuntimeRDG/Source/NNERuntimeRDG/Private/Hlsl/NNERuntimeRDGConvTranspose.cpp
@@ -205,8 +205,9 @@ namespace UE::NNERuntimeRDG::Private::Hlsl
 
 	bool RegisterConvTransposeOperator(FOperatorRegistryHlsl& Registry)
 	{
-		Registry.OpAdd({{TEXT("ConvTranspose"), TEXT("Onnx")}}, FConvTranspose::Create, ValidateConvTransposeOperator);
-
+		// Note: support of a particular version is partial with respect to tensor data types (only the most typical ones are usually supported).
+		Registry.OpAdd({{TEXT("ConvTranspose"), TEXT("Onnx")}, 1}, FConvTranspose::Create, ValidateConvTransposeOperator);
+		Registry.OpAdd({{TEXT("ConvTranspose"), TEXT("Onnx")}, 11}, FConvTranspose::Create, ValidateConvTransposeOperator);
 		return true;
 	}
 
diff --git a/Engine/Plugins/Experimental/NNERuntimeRDG/Source/NNERuntimeRDG/Private/Hlsl/NNERuntimeRDGDropout.cpp b/Engine/Plugins/Experimental/NNERuntimeRDG/Source/NNERuntimeRDG/Private/Hlsl/NNERuntimeRDGDropout.cpp
--- a/Engine/Plugins/Experimental/NNERuntimeRDG/Source/NNERuntimeRDG/Private/Hlsl/NNERuntimeRDGDropout.cpp
+++ b/Engine/Plugins/Experimental/NNERuntimeRDG/Source/NNERuntimeRDG/Private/Hlsl/NNERuntimeRDGDropout.cpp
@@ -72,8 +72,6 @@ namespace UE::NNERuntimeRDG::Private::Hlsl
 	{
 		bool bIsValid = true;
 
-		//This match version 13 of the Dropout operator
-		//[REDACTED_PROFILE_URL]/onnx/blob/main/docs/Operators.md#Dropout
 		FAttributeValidator AttributeValidator;
 		AttributeValidator.AddOptional(TEXT("seed"), ENNEAttributeDataType::Int32); //Will be ignored, only useful in training mode.
 		bIsValid &= AttributeValidator.Validate(AttributeMap);
@@ -101,7 +99,9 @@ namespace UE::NNERuntimeRDG::Private::Hlsl
 
 	bool RegisterDropoutOperator(FOperatorRegistryHlsl& Registry)
 	{
-		Registry.OpAdd({{TEXT("Dropout"), TEXT("Onnx")}}, CreateDropoutOperator, ValidateDropoutOperator);
+		// Note: support of a particular version is partial with respect to tensor data types (only the most typical ones are usually supported).
+		Registry.OpAdd({{TEXT("Dropout"), TEXT("Onnx")}, 12}, CreateDropoutOperator, ValidateDropoutOperator);
+		Registry.OpAdd({{TEXT("Dropout"), TEXT("Onnx")}, 13}, CreateDropoutOperator, ValidateDropoutOperator);
 		return true;
 	}
 } // UE::NNERuntimeRDG::Private::Hlsl
diff --git a/Engine/Plugins/Experimental/NNERuntimeRDG/Source/NNERuntimeRDG/Private/Hlsl/NNERuntimeRDGElementWiseBinary.cpp b/Engine/Plugins/Experimental/NNERuntimeRDG/Source/NNERuntimeRDG/Private/Hlsl/NNERuntimeRDGElementWiseBinary.cpp
--- a/Engine/Plugins/Experimental/NNERuntimeRDG/Source/NNERuntimeRDG/Private/Hlsl/NNERuntimeRDGElementWiseBinary.cpp
+++ b/Engine/Plugins/Experimental/NNERuntimeRDG/Source/NNERuntimeRDG/Private/Hlsl/NNERuntimeRDGElementWiseBinary.cpp
@@ -153,22 +153,52 @@ namespace UE::NNERuntimeRDG::Private::Hlsl
 
 	bool RegisterElementWiseBinaryOperators(FOperatorRegistryHlsl& Registry)
 	{
-#define OP(Name) Registry.OpAdd({{TEXT(#Name), TEXT("Onnx")}}, CreateElementWiseBinaryOperator<NNE::Internal::EElementWiseBinaryOperatorType::Name>, ValidateElementWiseBinaryOperator)
-		OP(Add);
-		//OP(And);
-		OP(Div);
-		//OP(Equal);
-		//OP(Greater);
-		//OP(GreaterOrEqual);
-		//OP(Less);
-		//OP(LessOrEqual);
-		OP(Mod);
-		OP(Mul);
-		//OP(Or);
-		OP(Prelu);
-		OP(Pow);
-		OP(Sub);
-		//OP(Xor);
+		// Note: support of a particular version is partial with respect to tensor data types (only the most typical ones are usually supported).
+#define OP(Name, Version) Registry.OpAdd({{TEXT(#Name), TEXT("Onnx")}, Version}, CreateElementWiseBinaryOperator<NNE::Internal::EElementWiseBinaryOperatorType::Name>, ValidateElementWiseBinaryOperator);
+#define OP_ALL_VERSIONS(Name) \
+OP(Name, 6) \
+OP(Name, 7) \
+OP(Name, 13) \
+OP(Name, 14)
+		
+		OP_ALL_VERSIONS(Add)
+		//OP(And, 1)
+		//OP(And, 7)
+		OP_ALL_VERSIONS(Div)
+		//OP(Equal, 1)
+		//OP(Equal, 7)
+		//OP(Equal, 11)
+		//OP(Equal, 13)
+		//OP(Equal, 19)
+		//OP(Greater, 1)
+		//OP(Greater, 7)
+		//OP(Greater, 9)
+		//OP(Greater, 13)
+		//OP(GreaterOrEqual, 12)
+		//OP(GreaterOrEqual, 16)
+		//OP(Less, 1)
+		//OP(Less, 7)
+		//OP(Less, 9)
+		//OP(Less, 13)
+		//OP(LessOrEqual, 12)
+		//OP(LessOrEqual, 16)
+		OP(Mod, 10)
+		OP(Mod, 13)
+		OP_ALL_VERSIONS(Mul)
+		//OP(Or, 1)
+		//OP(Or, 7)
+		OP(Prelu, 6)
+		OP(Prelu, 7)
+		OP(Prelu, 9)
+		OP(Prelu, 16)
+		OP(Pow, 7)
+		OP(Pow, 12)
+		OP(Pow, 13)
+		OP(Pow, 15)
+		OP_ALL_VERSIONS(Sub)
+		//OP(Xor, 1)
+		//OP(Xor, 7)
+#undef OP_ALL_VERSIONS
 #undef OP
 
 		return true;
diff --git a/Engine/Plugins/Experimental/NNERuntimeRDG/Source/NNERuntimeRDG/Private/Hlsl/NNERuntimeRDGElementWiseUnary.cpp b/Engine/Plugins/Experimental/NNERuntimeRDG/Source/NNERuntimeRDG/Private/Hlsl/NNERuntimeRDGElementWiseUnary.cpp
--- a/Engine/Plugins/Experimental/NNERuntimeRDG/Source/NNERuntimeRDG/Private/Hlsl/NNERuntimeRDGElementWiseUnary.cpp
+++ b/Engine/Plugins/Experimental/NNERuntimeRDG/Source/NNERuntimeRDG/Private/Hlsl/NNERuntimeRDGElementWiseUnary.cpp
@@ -224,8 +224,6 @@ namespace UE::NNERuntimeRDG::Private::Hlsl
 	{
 		bool bIsValid = true;
 
-		//version 6 of operator Clip, next version is 11
-		//[REDACTED_PROFILE_URL]/onnx/blob/main/docs/Changelog.md#Clip-6
 		FAttributeValidator AttributeValidator;
 		AttributeValidator.AddOptional(TEXT("min"), ENNEAttributeDataType::Float);
 		AttributeValidator.AddOptional(TEXT("max"), ENNEAttributeDataType::Float);
@@ -241,44 +239,63 @@ namespace UE::NNERuntimeRDG::Private::Hlsl
 	
 	bool RegisterElementWiseUnaryOperators(FOperatorRegistryHlsl& Registry)
 	{
-#define OP(Name) Registry.OpAdd({{TEXT(#Name), TEXT("Onnx")}}, CreateElementWiseUnaryOperator<NNE::Internal::EElementWiseUnaryOperatorType::Name>, ValidateElementWiseUnaryOperator<NNE::Internal::EElementWiseUnaryOperatorType::Name>)
-		OP(Abs);
-		OP(Acos);
-		OP(Acosh);
-		OP(Asin);
-		OP(Asinh);
-		OP(Atan);
-		OP(Atanh);
-		//OP(BitShift);
-		OP(Ceil);
-		OP(Clip);
-		OP(Cos);
-		OP(Cosh);
-		OP(Elu);
-		OP(Erf);
-		OP(Exp);
-		OP(Floor);
-		OP(IsInf);
-		OP(IsNan);
-		OP(HardSigmoid);
-		OP(HardSwish);
-		OP(LeakyRelu);
-		OP(Log);
-		OP(Neg);
-		//OP(Not);
-		OP(Reciprocal);
-		OP(Relu);
-		OP(Round);
-		OP(Selu);
-		OP(Sigmoid);
-		OP(Sign);
-		OP(Sin);
-		OP(Sinh);
-		OP(Softplus);
-		OP(Softsign);
-		OP(Sqrt);
-		OP(Tan);
-		OP(Tanh);
+		// Note: support of a particular version is partial with respect to tensor data types (only the most typical ones are usually supported).
+#define OP(Name, Version) Registry.OpAdd({{TEXT(#Name), TEXT("Onnx")}, Version}, CreateElementWiseUnaryOperator<NNE::Internal::EElementWiseUnaryOperatorType::Name>, ValidateElementWiseUnaryOperator<NNE::Internal::EElementWiseUnaryOperatorType::Name>);
+		OP(Abs, 6)
+		OP(Abs, 13)
+		OP(Acos, 7)
+		OP(Acosh, 9)
+		OP(Asin, 7)
+		OP(Asinh, 9)
+		OP(Atan, 7)
+		OP(Atanh, 9)
+		//OP(BitShift, 11)
+		OP(Ceil, 6)
+		OP(Ceil, 13)
+		OP(Clip, 6)
+		OP(Cos, 7)
+		OP(Cosh, 9)
+		OP(Elu, 6)
+		OP(Erf, 9)
+		OP(Erf, 13)
+		OP(Exp, 6)
+		OP(Exp, 13)
+		OP(Floor, 6)
+		OP(Floor, 13)
+		OP(IsInf, 10)
+		OP(IsInf, 20)
+		OP(IsNan, 9)
+		OP(IsNan, 13)
+		OP(IsNan, 20)
+		OP(HardSigmoid, 6)
+		OP(HardSwish, 14)
+		OP(LeakyRelu, 6)
+		OP(LeakyRelu, 16)
+		OP(Log, 6)
+		OP(Log, 13)
+		OP(Neg, 6)
+		OP(Neg, 13)
+		//OP(Not, 1)
+		OP(Reciprocal, 6)
+		OP(Reciprocal, 13)
+		OP(Relu, 6)
+		OP(Relu, 13)
+		OP(Relu, 14)
+		OP(Round, 11)
+		OP(Selu, 6)
+		OP(Sigmoid, 6)
+		OP(Sigmoid, 13)
+		OP(Sign, 9)
+		OP(Sign, 13)
+		OP(Sin, 7)
+		OP(Sinh, 9)
+		OP(Softplus, 1)
+		OP(Softsign, 1)
+		OP(Sqrt, 6)
+		OP(Sqrt, 13)
+		OP(Tan, 7)
+		OP(Tanh, 6)
+		OP(Tanh, 13)
 #undef OP
 
 		return true;
diff --git a/Engine/Plugins/Experimental/NNERuntimeRDG/Source/NNERuntimeRDG/Private/Hlsl/NNERuntimeRDGElementWiseVariadic.cpp b/Engine/Plugins/Experimental/NNERuntimeRDG/Source/NNERuntimeRDG/Private/Hlsl/NNERuntimeRDGElementWiseVariadic.cpp
--- a/Engine/Plugins/Experimental/NNERuntimeRDG/Source/NNERuntimeRDG/Private/Hlsl/NNERuntimeRDGElementWiseVariadic.cpp
+++ b/Engine/Plugins/Experimental/NNERuntimeRDG/Source/NNERuntimeRDG/Private/Hlsl/NNERuntimeRDGElementWiseVariadic.cpp
@@ -211,11 +211,21 @@ namespace UE::NNERuntimeRDG::Private::Hlsl
 
 	bool RegisterElementWiseVariadicOperators(FOperatorRegistryHlsl& Registry)
 	{
-#define OP(Name) Registry.OpAdd({{TEXT(#Name), TEXT("Onnx")}}, CreateElementWiseVariadicOperator<NNE::Internal::EElementWiseVariadicOperatorType::Name>)
-		OP(Max);
-		OP(Min);
-		OP(Mean);
-		OP(Sum);
+		// Note: support of a particular version is partial with respect to tensor data types (only the most typical ones are usually supported).
+#define OP(Name, Version) Registry.OpAdd({{TEXT(#Name), TEXT("Onnx")}, Version}, CreateElementWiseVariadicOperator<NNE::Internal::EElementWiseVariadicOperatorType::Name>);
+#define OP_ALL_VERSIONS(Name) \
+OP(Name, 6) \
+OP(Name, 8) \
+OP(Name, 12) \
+OP(Name, 13)
+		OP_ALL_VERSIONS(Max)
+		OP_ALL_VERSIONS(Min)
+		OP(Mean, 6)
+		OP(Mean, 8)
+		OP(Mean, 13)
+		OP(Sum, 6)
+		OP(Sum, 8)
+		OP(Sum, 13)
 #undef OP
 
 		return true;
diff --git a/Engine/Plugins/Experimental/NNERuntimeRDG/Source/NNERuntimeRDG/Private/Hlsl/NNERuntimeRDGFlatten.cpp b/Engine/Plugins/Experimental/NNERuntimeRDG/Source/NNERuntimeRDG/Private/Hlsl/NNERuntimeRDGFlatten.cpp
--- a/Engine/Plugins/Experimental/NNERuntimeRDG/Source/NNERuntimeRDG/Private/Hlsl/NNERuntimeRDGFlatten.cpp
+++ b/Engine/Plugins/Experimental/NNERuntimeRDG/Source/NNERuntimeRDG/Private/Hlsl/NNERuntimeRDGFlatten.cpp
@@ -92,8 +92,6 @@ namespace UE::NNERuntimeRDG::Private::Hlsl
 	{
 		bool bIsValid = true;
 
-		//This match version 13 of the Flatten operator
-		//[REDACTED_PROFILE_URL]/onnx/blob/main/docs/Operators.md#Flatten
 		FAttributeValidator AttributeValidator;
 		AttributeValidator.AddOptional(TEXT("axis"), ENNEAttributeDataType::Int32);
 		bIsValid &= AttributeValidator.Validate(AttributeMap);
@@ -124,7 +122,11 @@ namespace UE::NNERuntimeRDG::Private::Hlsl
 
 	bool RegisterFlattenOperator(FOperatorRegistryHlsl& Registry)
 	{
-		Registry.OpAdd({{TEXT("Flatten"), TEXT("Onnx")}}, CreateFlattenOperator, ValidateFlattenOperator);
+		// Note: support of a particular version is partial with respect to tensor data types (only the most typical ones are usually supported).
+		Registry.OpAdd({{TEXT("Flatten"), TEXT("Onnx")}, 1}, CreateFlattenOperator, ValidateFlattenOperator);
+		Registry.OpAdd({{TEXT("Flatten"), TEXT("Onnx")}, 9}, CreateFlattenOperator, ValidateFlattenOperator);
+		Registry.OpAdd({{TEXT("Flatten"), TEXT("Onnx")}, 11}, CreateFlattenOperator, ValidateFlattenOperator);
+		Registry.OpAdd({{TEXT("Flatten"), TEXT("Onnx")}, 13}, CreateFlattenOperator, ValidateFlattenOperator);
 		return true;
 	}
 } // UE::NNERuntimeRDG::Private::Hlsl
diff --git a/Engine/Plugins/Experimental/NNERuntimeRDG/Source/NNERuntimeRDG/Private/Hlsl/NNERuntimeRDGGather.cpp b/Engine/Plugins/Experimental/NNERuntimeRDG/Source/NNERuntimeRDG/Private/Hlsl/NNERuntimeRDGGather.cpp
--- a/Engine/Plugins/Experimental/NNERuntimeRDG/Source/NNERuntimeRDG/Private/Hlsl/NNERuntimeRDGGather.cpp
+++ b/Engine/Plugins/Experimental/NNERuntimeRDG/Source/NNERuntimeRDG/Private/Hlsl/NNERuntimeRDGGather.cpp
@@ -155,8 +155,6 @@ namespace UE::NNERuntimeRDG::Private::Hlsl
 	{
 		bool bIsValid = true;
 
-		// This matches versions 1, 11 and 13 of the Gather operator
-		// [REDACTED_PROFILE_URL]/onnx/blob/main/docs/Operators.md#Gather
 		FAttributeValidator AttributeValidator;
 		AttributeValidator.AddOptional(TEXT("axis"), ENNEAttributeDataType::Int32);
 		bIsValid &= AttributeValidator.Validate(AttributeMap);
@@ -184,7 +182,10 @@ namespace UE::NNERuntimeRDG::Private::Hlsl
 
 	bool RegisterGatherOperator(FOperatorRegistryHlsl& Registry)
 	{
-		Registry.OpAdd({{TEXT("Gather"), TEXT("Onnx")}}, CreateGatherOperator, ValidateGatherOperator);
+		// Note: support of a particular version is partial with respect to tensor data types (only the most typical ones are usually supported).
+		Registry.OpAdd({{TEXT("Gather"), TEXT("Onnx")}, 1}, CreateGatherOperator, ValidateGatherOperator);
+		Registry.OpAdd({{TEXT("Gather"), TEXT("Onnx")}, 11}, CreateGatherOperator, ValidateGatherOperator);
+		Registry.OpAdd({{TEXT("Gather"), TEXT("Onnx")}, 13}, CreateGatherOperator, ValidateGatherOperator);
 		return true;
 	}
 } // UE::NNERuntimeRDG::Private::Hlsl
\ No newline at end of file
diff --git a/Engine/Plugins/Experimental/NNERuntimeRDG/Source/NNERuntimeRDG/Private/Hlsl/NNERuntimeRDGGemm.cpp b/Engine/Plugins/Experimental/NNERuntimeRDG/Source/NNERuntimeRDG/Private/Hlsl/NNERuntimeRDGGemm.cpp
--- a/Engine/Plugins/Experimental/NNERuntimeRDG/Source/NNERuntimeRDG/Private/Hlsl/NNERuntimeRDGGemm.cpp
+++ b/Engine/Plugins/Experimental/NNERuntimeRDG/Source/NNERuntimeRDG/Private/Hlsl/NNERuntimeRDGGemm.cpp
@@ -190,7 +190,11 @@ namespace UE::NNERuntimeRDG::Private::Hlsl
 
 	bool RegisterGemmOperator(FOperatorRegistryHlsl& Registry)
 	{
-		Registry.OpAdd({{TEXT("Gemm"), TEXT("Onnx")}}, CreateGemmOperator, ValidateGemmOperator);
+		// Note: support of a particular version is partial with respect to tensor data types (only the most typical ones are usually supported).
+		Registry.OpAdd({{TEXT("Gemm"), TEXT("Onnx")}, 7}, CreateGemmOperator, ValidateGemmOperator);
+		Registry.OpAdd({{TEXT("Gemm"), TEXT("Onnx")}, 9}, CreateGemmOperator, ValidateGemmOperator);
+		Registry.OpAdd({{TEXT("Gemm"), TEXT("Onnx")}, 11}, CreateGemmOperator, ValidateGemmOperator);
+		Registry.OpAdd({{TEXT("Gemm"), TEXT("Onnx")}, 13}, CreateGemmOperator, ValidateGemmOperator);
 		return true;
 	}
 } // UE::NNERuntimeRDG::Private::Hlsl
diff --git a/Engine/Plugins/Experimental/NNERuntimeRDG/Source/NNERuntimeRDG/Private/Hlsl/NNERuntimeRDGIdentity.cpp b/Engine/Plugins/Experimental/NNERuntimeRDG/Source/NNERuntimeRDG/Private/Hlsl/NNERuntimeRDGIdentity.cpp
--- a/Engine/Plugins/Experimental/NNERuntimeRDG/Source/NNERuntimeRDG/Private/Hlsl/NNERuntimeRDGIdentity.cpp
+++ b/Engine/Plugins/Experimental/NNERuntimeRDG/Source/NNERuntimeRDG/Private/Hlsl/NNERuntimeRDGIdentity.cpp
@@ -66,8 +66,6 @@ namespace UE::NNERuntimeRDG::Private::Hlsl
 	{
 		bool bIsValid = true;
 
-		//This match version 16 of the Identity operator
-		//[REDACTED_PROFILE_URL]/onnx/blob/main/docs/Operators.md#Identity
 		FAttributeValidator AttributeValidator;
 		bIsValid &= AttributeValidator.Validate(AttributeMap);
 
@@ -97,7 +95,12 @@ namespace UE::NNERuntimeRDG::Private::Hlsl
 
 	bool RegisterIdentityOperator(FOperatorRegistryHlsl& Registry)
 	{
-		Registry.OpAdd({{TEXT("Identity"), TEXT("Onnx")}}, CreateIdentityOperator, ValidateIdentityOperator);
+		// Note: support of a particular version is partial with respect to tensor data types (only the most typical ones are usually supported).
+		Registry.OpAdd({{TEXT("Identity"), TEXT("Onnx")}, 1}, CreateIdentityOperator, ValidateIdentityOperator);
+		Registry.OpAdd({{TEXT("Identity"), TEXT("Onnx")}, 13}, CreateIdentityOperator, ValidateIdentityOperator);
+		Registry.OpAdd({{TEXT("Identity"), TEXT("Onnx")}, 14}, CreateIdentityOperator, ValidateIdentityOperator);
+		Registry.OpAdd({{TEXT("Identity"), TEXT("Onnx")}, 16}, CreateIdentityOperator, ValidateIdentityOperator);
+		Registry.OpAdd({{TEXT("Identity"), TEXT("Onnx")}, 19}, CreateIdentityOperator, ValidateIdentityOperator);
 		return true;
 	}
 } // UE::NNERuntimeRDG::Private::Hlsl
diff --git a/Engine/Plugins/Experimental/NNERuntimeRDG/Source/NNERuntimeRDG/Private/Hlsl/NNERuntimeRDGInstanceNormalization.cpp b/Engine/Plugins/Experimental/NNERuntimeRDG/Source/NNERuntimeRDG/Private/Hlsl/NNERuntimeRDGInstanceNormalization.cpp
--- a/Engine/Plugins/Experimental/NNERuntimeRDG/Source/NNERuntimeRDG/Private/Hlsl/NNERuntimeRDGInstanceNormalization.cpp
+++ b/Engine/Plugins/Experimental/NNERuntimeRDG/Source/NNERuntimeRDG/Private/Hlsl/NNERuntimeRDGInstanceNormalization.cpp
@@ -172,7 +172,8 @@ namespace UE::NNERuntimeRDG::Private::Hlsl
 
 	bool RegisterInstanceNormalizationOperator(FOperatorRegistryHlsl& Registry)
 	{
-		Registry.OpAdd({{TEXT("InstanceNormalization"), TEXT("Onnx")}}, CreateInstanceNormalizationOperator, ValidateInstanceNormalizationOperator);
+		// Note: support of a particular version is partial with respect to tensor data types (only the most typical ones are usually supported).
+		Registry.OpAdd({{TEXT("InstanceNormalization"), TEXT("Onnx")}, 6}, CreateInstanceNormalizationOperator, ValidateInstanceNormalizationOperator);
 		return true;
 	}
 } // UE::NNERuntimeRDG::Private::Hlsl
diff --git a/Engine/Plugins/Experimental/NNERuntimeRDG/Source/NNERuntimeRDG/Private/Hlsl/NNERuntimeRDGMatMul.cpp b/Engine/Plugins/Experimental/NNERuntimeRDG/Source/NNERuntimeRDG/Private/Hlsl/NNERuntimeRDGMatMul.cpp
--- a/Engine/Plugins/Experimental/NNERuntimeRDG/Source/NNERuntimeRDG/Private/Hlsl/NNERuntimeRDGMatMul.cpp
+++ b/Engine/Plugins/Experimental/NNERuntimeRDG/Source/NNERuntimeRDG/Private/Hlsl/NNERuntimeRDGMatMul.cpp
@@ -162,7 +162,10 @@ namespace UE::NNERuntimeRDG::Private::Hlsl
 
 	bool RegisterMatMulOperator(FOperatorRegistryHlsl& Registry)
 	{
-		Registry.OpAdd({{TEXT("MatMul"), TEXT("Onnx")}}, CreateMatMulOperator, ValidateMatMulOperator);
+		// Note: support of a particular version is partial with respect to tensor data types (only the most typical ones are usually supported).
+		Registry.OpAdd({{TEXT("MatMul"), TEXT("Onnx")}, 1}, CreateMatMulOperator, ValidateMatMulOperator);
+		Registry.OpAdd({{TEXT("MatMul"), TEXT("Onnx")}, 9}, CreateMatMulOperator, ValidateMatMulOperator);
+		Registry.OpAdd({{TEXT("MatMul"), TEXT("Onnx")}, 13}, CreateMatMulOperator, ValidateMatMulOperator);
 		return true;
 	}
 } // UE::NNERuntimeRDG::Private::Hlsl
diff --git a/Engine/Plugins/Experimental/NNERuntimeRDG/Source/NNERuntimeRDG/Private/Hlsl/NNERuntimeRDGPad.cpp b/Engine/Plugins/Experimental/NNERuntimeRDG/Source/NNERuntimeRDG/Private/Hlsl/NNERuntimeRDGPad.cpp
--- a/Engine/Plugins/Experimental/NNERuntimeRDG/Source/NNERuntimeRDG/Private/Hlsl/NNERuntimeRDGPad.cpp
+++ b/Engine/Plugins/Experimental/NNERuntimeRDG/Source/NNERuntimeRDG/Private/Hlsl/NNERuntimeRDGPad.cpp
@@ -129,8 +129,6 @@ namespace UE::NNERuntimeRDG::Private::Hlsl
 	{
 		bool bIsValid = true;
 
-		//This match version 2 of the pad operator (next version is with opset 11)
-		//[REDACTED_PROFILE_URL]/onnx/blob/main/docs/Changelog.md#Pad-2
 		FAttributeValidator AttributeValidator;
 		AttributeValidator.AddOptional(TEXT("mode"), ENNEAttributeDataType::String);
 		AttributeValidator.AddRequired(TEXT("pads"), ENNEAttributeDataType::Int32Array);
@@ -161,7 +159,8 @@ namespace UE::NNERuntimeRDG::Private::Hlsl
 
 	bool RegisterPadOperator(FOperatorRegistryHlsl& Registry)
 	{
-		Registry.OpAdd({{TEXT("Pad"), TEXT("Onnx")}}, CreatePadOperator, ValidatePadOperator);
+		// Note: support of a particular version is partial with respect to tensor data types (only the most typical ones are usually supported).
+		Registry.OpAdd({{TEXT("Pad"), TEXT("Onnx")}, 2}, CreatePadOperator, ValidatePadOperator);
 		return true;
 	}
 } // UE::NNERuntimeRDG::Private::Hlsl
diff --git a/Engine/Plugins/Experimental/NNERuntimeRDG/Source/NNERuntimeRDG/Private/Hlsl/NNERuntimeRDGReshape.cpp b/Engine/Plugins/Experimental/NNERuntimeRDG/Source/NNERuntimeRDG/Private/Hlsl/NNERuntimeRDGReshape.cpp
--- a/Engine/Plugins/Experimental/NNERuntimeRDG/Source/NNERuntimeRDG/Private/Hlsl/NNERuntimeRDGReshape.cpp
+++ b/Engine/Plugins/Experimental/NNERuntimeRDG/Source/NNERuntimeRDG/Private/Hlsl/NNERuntimeRDGReshape.cpp
@@ -88,8 +88,6 @@ namespace UE::NNERuntimeRDG::Private::Hlsl
 	{
 		bool bIsValid = true;
 
-		//This match version 14 of the Reshape operator
-		//[REDACTED_PROFILE_URL]/onnx/blob/main/docs/Operators.md#Reshape
 		FAttributeValidator AttributeValidator;
 		AttributeValidator.AddOptional(TEXT("allowzero"), ENNEAttributeDataType::Int32);
 		bIsValid &= AttributeValidator.Validate(AttributeMap);
@@ -124,7 +122,11 @@ namespace UE::NNERuntimeRDG::Private::Hlsl
 
 	bool RegisterReshapeOperator(FOperatorRegistryHlsl& Registry)
 	{
-		Registry.OpAdd({{TEXT("Reshape"), TEXT("Onnx")}}, CreateReshapeOperator, ValidateReshapeOperator);
+		// Note: support of a particular version is partial with respect to tensor data types (only the most typical ones are usually supported).
+		Registry.OpAdd({{TEXT("Reshape"), TEXT("Onnx")}, 5}, CreateReshapeOperator, ValidateReshapeOperator);
+		Registry.OpAdd({{TEXT("Reshape"), TEXT("Onnx")}, 13}, CreateReshapeOperator, ValidateReshapeOperator);
+		Registry.OpAdd({{TEXT("Reshape"), TEXT("Onnx")}, 14}, CreateReshapeOperator, ValidateReshapeOperator);
+		Registry.OpAdd({{TEXT("Reshape"), TEXT("Onnx")}, 19}, CreateReshapeOperator, ValidateReshapeOperator);
 		return true;
 	}
 } // UE::NNERuntimeRDG::Private::Hlsl
diff --git a/Engine/Plugins/Experimental/NNERuntimeRDG/Source/NNERuntimeRDG/Private/Hlsl/NNERuntimeRDGShape.cpp b/Engine/Plugins/Experimental/NNERuntimeRDG/Source/NNERuntimeRDG/Private/Hlsl/NNERuntimeRDGShape.cpp
--- a/Engine/Plugins/Experimental/NNERuntimeRDG/Source/NNERuntimeRDG/Private/Hlsl/NNERuntimeRDGShape.cpp
+++ b/Engine/Plugins/Experimental/NNERuntimeRDG/Source/NNERuntimeRDG/Private/Hlsl/NNERuntimeRDGShape.cpp
@@ -67,8 +67,6 @@ namespace UE::NNERuntimeRDG::Private::Hlsl
 	{
 		bool bIsValid = true;
 
-		//This match version 1 of the shape operator (next version is with opset 13)
-		//[REDACTED_PROFILE_URL]/onnx/blob/main/docs/Changelog.md#Shape-1
 		FAttributeValidator AttributeValidator;
 		bIsValid &= AttributeValidator.Validate(AttributeMap);
 
@@ -97,7 +95,9 @@ namespace UE::NNERuntimeRDG::Private::Hlsl
 
 	bool RegisterShapeOperator(FOperatorRegistryHlsl& Registry)
 	{
-		Registry.OpAdd({{TEXT("Shape"), TEXT("Onnx")}}, CreateShapeOperator, ValidateShapeOperator);
+		// Note: support of a particular version is partial with respect to tensor data types (only the most typical ones are usually supported).
+		Registry.OpAdd({{TEXT("Shape"), TEXT("Onnx")}, 1}, CreateShapeOperator, ValidateShapeOperator);
+		Registry.OpAdd({{TEXT("Shape"), TEXT("Onnx")}, 13}, CreateShapeOperator, ValidateShapeOperator);
 		return true;
 	}
 } // UE::NNERuntimeRDG::Private::Hlsl
diff --git a/Engine/Plugins/Experimental/NNERuntimeRDG/Source/NNERuntimeRDG/Private/Hlsl/NNERuntimeRDGSize.cpp b/Engine/Plugins/Experimental/NNERuntimeRDG/Source/NNERuntimeRDG/Private/Hlsl/NNERuntimeRDGSize.cpp
--- a/Engine/Plugins/Experimental/NNERuntimeRDG/Source/NNERuntimeRDG/Private/Hlsl/NNERuntimeRDGSize.cpp
+++ b/Engine/Plugins/Experimental/NNERuntimeRDG/Source/NNERuntimeRDG/Private/Hlsl/NNERuntimeRDGSize.cpp
@@ -55,8 +55,6 @@ namespace UE::NNERuntimeRDG::Private::Hlsl
 	{
 		bool bIsValid = true;
 
-		//This match version 13 of the Size operator, next version are 11 and 13
-		//[REDACTED_PROFILE_URL]/onnx/blob/main/docs/Operators.md#Size
 		FAttributeValidator AttributeValidator;
 		bIsValid &= AttributeValidator.Validate(AttributeMap);
 
@@ -86,7 +84,10 @@ namespace UE::NNERuntimeRDG::Private::Hlsl
 
 	bool RegisterSizeOperator(FOperatorRegistryHlsl& Registry)
 	{
-		Registry.OpAdd({{TEXT("Size"), TEXT("Onnx")}}, CreateSizeOperator, ValidateSizeOperator);
+		// Note: support of a particular version is partial with respect to tensor data types (only the most typical ones are usually supported).
+		Registry.OpAdd({{TEXT("Size"), TEXT("Onnx")}, 1}, CreateSizeOperator, ValidateSizeOperator);
+		Registry.OpAdd({{TEXT("Size"), TEXT("Onnx")}, 13}, CreateSizeOperator, ValidateSizeOperator);
+		Registry.OpAdd({{TEXT("Size"), TEXT("Onnx")}, 19}, CreateSizeOperator, ValidateSizeOperator);
 		return true;
 	}
 } // UE::NNERuntimeRDG::Private::Hlsl
diff --git a/Engine/Plugins/Experimental/NNERuntimeRDG/Source/NNERuntimeRDG/Private/Hlsl/NNERuntimeRDGSlice.cpp b/Engine/Plugins/Experimental/NNERuntimeRDG/Source/NNERuntimeRDG/Private/Hlsl/NNERuntimeRDGSlice.cpp
--- a/Engine/Plugins/Experimental/NNERuntimeRDG/Source/NNERuntimeRDG/Private/Hlsl/NNERuntimeRDGSlice.cpp
+++ b/Engine/Plugins/Experimental/NNERuntimeRDG/Source/NNERuntimeRDG/Private/Hlsl/NNERuntimeRDGSlice.cpp
@@ -195,8 +195,6 @@ namespace UE::NNERuntimeRDG::Private::Hlsl
 	{
 		bool bIsValid = true;
 
-		//This match version 1 of the Slice operator, next version are 10, 11 and 13
-		//[REDACTED_PROFILE_URL]/onnx/blob/main/docs/Changelog.md#Slice-1
 		FAttributeValidator AttributeValidator;
 		AttributeValidator.AddOptional(TEXT("axes"), ENNEAttributeDataType::Int32Array);
 		AttributeValidator.AddRequired(TEXT("ends"), ENNEAttributeDataType::Int32Array);
@@ -221,7 +219,8 @@ namespace UE::NNERuntimeRDG::Private::Hlsl
 
 	bool RegisterSliceOperator(FOperatorRegistryHlsl& Registry)
 	{
-		Registry.OpAdd({{TEXT("Slice"), TEXT("Onnx")}}, CreateSliceOperator, ValidateSliceOperator);
+		// Note: support of a particular version is partial with respect to tensor data types (only the most typical ones are usually supported).
+		Registry.OpAdd({{TEXT("Slice"), TEXT("Onnx")}, 1}, CreateSliceOperator, ValidateSliceOperator);
 		return true;
 	}
 } // UE::NNERuntimeRDG::Private::Hlsl
diff --git a/Engine/Plugins/Experimental/NNERuntimeRDG/Source/NNERuntimeRDG/Private/Hlsl/NNERuntimeRDGSqueeze.cpp b/Engine/Plugins/Experimental/NNERuntimeRDG/Source/NNERuntimeRDG/Private/Hlsl/NNERuntimeRDGSqueeze.cpp
--- a/Engine/Plugins/Experimental/NNERuntimeRDG/Source/NNERuntimeRDG/Private/Hlsl/NNERuntimeRDGSqueeze.cpp
+++ b/Engine/Plugins/Experimental/NNERuntimeRDG/Source/NNERuntimeRDG/Private/Hlsl/NNERuntimeRDGSqueeze.cpp
@@ -101,8 +101,6 @@ namespace UE::NNERuntimeRDG::Private::Hlsl
 	{
 		bool bIsValid = true;
 
-		//This match version 1 of the Squeeze operator, next version are 11 and 13
-		//[REDACTED_PROFILE_URL]/onnx/blob/main/docs/Operators.md#Squeeze
 		FAttributeValidator AttributeValidator;
 		AttributeValidator.AddRequired(TEXT("axes"), ENNEAttributeDataType::Int32Array);
 		bIsValid &= AttributeValidator.Validate(AttributeMap);
@@ -133,7 +131,8 @@ namespace UE::NNERuntimeRDG::Private::Hlsl
 
 	bool RegisterSqueezeOperator(FOperatorRegistryHlsl& Registry)
 	{
-		Registry.OpAdd({{TEXT("Squeeze"), TEXT("Onnx")}}, CreateSqueezeOperator, ValidateSqueezeOperator);
+		// Note: support of a particular version is partial with respect to tensor data types (only the most typical ones are usually supported).
+		Registry.OpAdd({{TEXT("Squeeze"), TEXT("Onnx")}, 1}, CreateSqueezeOperator, ValidateSqueezeOperator);
 		return true;
 	}
 } // UE::NNERuntimeRDG::Private::Hlsl
diff --git a/Engine/Plugins/Experimental/NNERuntimeRDG/Source/NNERuntimeRDG/Private/Hlsl/NNERuntimeRDGTranspose.cpp b/Engine/Plugins/Experimental/NNERuntimeRDG/Source/NNERuntimeRDG/Private/Hlsl/NNERuntimeRDGTranspose.cpp
--- a/Engine/Plugins/Experimental/NNERuntimeRDG/Source/NNERuntimeRDG/Private/Hlsl/NNERuntimeRDGTranspose.cpp
+++ b/Engine/Plugins/Experimental/NNERuntimeRDG/Source/NNERuntimeRDG/Private/Hlsl/NNERuntimeRDGTranspose.cpp
@@ -132,8 +132,6 @@ namespace UE::NNERuntimeRDG::Private::Hlsl
 	{
 		bool bIsValid = true;
 
-		//This match version 1 and 13 of the Transpose operator
-		//[REDACTED_PROFILE_URL]/onnx/blob/main/docs/Operators.md#Transpose
 		FAttributeValidator AttributeValidator;
 		AttributeValidator.AddOptional(TEXT("perm"), ENNEAttributeDataType::Int32Array);
 		bIsValid &= AttributeValidator.Validate(AttributeMap);
@@ -153,7 +151,9 @@ namespace UE::NNERuntimeRDG::Private::Hlsl
 
 	bool RegisterTransposeOperator(FOperatorRegistryHlsl& Registry)
 	{
-		Registry.OpAdd({{TEXT("Transpose"), TEXT("Onnx")}}, CreateTransposeOperator, ValidateTransposeOperator);
+		// Note: support of a particular version is partial with respect to tensor data types (only the most typical ones are usually supported).
+		Registry.OpAdd({{TEXT("Transpose"), TEXT("Onnx")}, 1}, CreateTransposeOperator, ValidateTransposeOperator);
+		Registry.OpAdd({{TEXT("Transpose"), TEXT("Onnx")}, 13}, CreateTransposeOperator, ValidateTransposeOperator);
 		return true;
 	}
 } // UE::NNERuntimeRDG::Private::Hlsl
diff --git a/Engine/Plugins/Experimental/NNERuntimeRDG/Source/NNERuntimeRDG/Private/Hlsl/NNERuntimeRDGUnsqueeze.cpp b/Engine/Plugins/Experimental/NNERuntimeRDG/Source/NNERuntimeRDG/Private/Hlsl/NNERuntimeRDGUnsqueeze.cpp
--- a/Engine/Plugins/Experimental/NNERuntimeRDG/Source/NNERuntimeRDG/Private/Hlsl/NNERuntimeRDGUnsqueeze.cpp
+++ b/Engine/Plugins/Experimental/NNERuntimeRDG/Source/NNERuntimeRDG/Private/Hlsl/NNERuntimeRDGUnsqueeze.cpp
@@ -93,8 +93,6 @@ namespace UE::NNERuntimeRDG::Private::Hlsl
 	{
 		bool bIsValid = true;
 
-		//This match version 1 of the Unsqueeze operator, next version are 11 and 13
-		//[REDACTED_PROFILE_URL]/onnx/blob/main/docs/Operators.md#Unsqueeze
 		FAttributeValidator AttributeValidator;
 		AttributeValidator.AddRequired(TEXT("axes"), ENNEAttributeDataType::Int32Array);
 		bIsValid &= AttributeValidator.Validate(AttributeMap);
@@ -125,7 +123,8 @@ namespace UE::NNERuntimeRDG::Private::Hlsl
 
 	bool RegisterUnsqueezeOperator(FOperatorRegistryHlsl& Registry)
 	{
-		Registry.OpAdd({{TEXT("Unsqueeze"), TEXT("Onnx")}}, CreateUnsqueezeOperator, ValidateUnsqueezeOperator);
+		// Note: support of a particular version is partial with respect to tensor data types (only the most typical ones are usually supported).
+		Registry.OpAdd({{TEXT("Unsqueeze"), TEXT("Onnx")}, 1}, CreateUnsqueezeOperator, ValidateUnsqueezeOperator);
 		return true;
 	}
 } // UE::NNERuntimeRDG::Private::Hlsl
diff --git a/Engine/Plugins/Experimental/NNERuntimeRDG/Source/NNERuntimeRDG/Private/Hlsl/NNERuntimeRDGUpsample.cpp b/Engine/Plugins/Experimental/NNERuntimeRDG/Source/NNERuntimeRDG/Private/Hlsl/NNERuntimeRDGUpsample.cpp
--- a/Engine/Plugins/Experimental/NNERuntimeRDG/Source/NNERuntimeRDG/Private/Hlsl/NNERuntimeRDGUpsample.cpp
+++ b/Engine/Plugins/Experimental/NNERuntimeRDG/Source/NNERuntimeRDG/Private/Hlsl/NNERuntimeRDGUpsample.cpp
@@ -147,7 +147,8 @@ namespace UE::NNERuntimeRDG::Private::Hlsl
 
 	bool RegisterUpsampleOperator(FOperatorRegistryHlsl& Registry)
 	{
-		Registry.OpAdd({{TEXT("Upsample"), TEXT("Onnx")}}, CreateUpsampleOperator, ValidateUpsampleOperator);
+		// Note: support of a particular version is partial with respect to tensor data types (only the most typical ones are usually supported).
+		Registry.OpAdd({{TEXT("Upsample"), TEXT("Onnx")}, 9}, CreateUpsampleOperator, ValidateUpsampleOperator);
 		return true;
 	}
 } // UE::NNERuntimeRDG::Private::Hlsl
