[NNE] NNERuntimeHLSL Use factorized Reduction operation from Reduce for InstanceNorm

Problem:
InstanceNorm was doing inplace reduction rather than using the shared ReduceCS one, not benefiting for further potential optimization applyed to it.

Solution:
Refactor the op to first apply the generic reduction then InstanceNorm in an element wise way.

Test: Win64 Editor: smoke tests + instance norm tests
Test: Win64 Standalone : all tests

#rnx

#virtualized

diff --git a/Engine/Plugins/Experimental/NNERuntimeRDG/Source/NNEHlslShaders/Internal/NNEHlslShadersInstanceNormalizationCS.h b/Engine/Plugins/Experimental/NNERuntimeRDG/Source/NNEHlslShaders/Internal/NNEHlslShadersInstanceNormalizationCS.h
--- a/Engine/Plugins/Experimental/NNERuntimeRDG/Source/NNEHlslShaders/Internal/NNEHlslShadersInstanceNormalizationCS.h
+++ b/Engine/Plugins/Experimental/NNERuntimeRDG/Source/NNEHlslShaders/Internal/NNEHlslShadersInstanceNormalizationCS.h
@@ -10,14 +10,10 @@ namespace UE::NNE::Internal { class FTensor; }
 
 namespace UE::NNEHlslShaders::Internal
 {
-	enum class EInstanceNormalizationAlgorithm : uint8
+	class FInstanceNormalizationConstants
 	{
-		Simple1x265 = 0,
-		SharedMemory256x1,
-		SharedMemory512x1,
-		SharedMemory768x1,
-		SharedMemory1024x1,
-		MAX
+	public:
+		static const int32 NUM_GROUP_THREADS{ 256 };
 	};
 
 	class NNEHLSLSHADERS_API TInstanceNormalizationCS : public FHlslShaderBase
@@ -25,27 +21,21 @@ namespace UE::NNEHlslShaders::Internal
 		DECLARE_GLOBAL_SHADER(TInstanceNormalizationCS);
 		SHADER_USE_PARAMETER_STRUCT(TInstanceNormalizationCS, FHlslShaderBase)
 
-		class FInstanceNormalizationAlgorithm : SHADER_PERMUTATION_ENUM_CLASS("ALGORITHM", EInstanceNormalizationAlgorithm);
-		using FPermutationDomain = TShaderPermutationDomain<FInstanceNormalizationAlgorithm>;
-
 	public:
 
 		BEGIN_SHADER_PARAMETER_STRUCT(FParameters, )
-			SHADER_PARAMETER(float, Epsilon)
-			SHADER_PARAMETER(int32, C)
-			SHADER_PARAMETER(int32, NxC)
-			SHADER_PARAMETER(int32, W)
+			SHADER_PARAMETER(uint32, Num)
+			SHADER_PARAMETER(uint32, ThreadCountX)
+			SHADER_PARAMETER(uint32, InstanceSize)
+			SHADER_PARAMETER(uint32, ChannelSize)
 			SHADER_PARAMETER_RDG_BUFFER_SRV(Buffer<float>, Input)
-			SHADER_PARAMETER_RDG_BUFFER_SRV(Buffer<float>, Scale)
-			SHADER_PARAMETER_RDG_BUFFER_SRV(Buffer<float>, Bias)
+			SHADER_PARAMETER_RDG_BUFFER_SRV(Buffer<float>, InputScale)
+			SHADER_PARAMETER_RDG_BUFFER_SRV(Buffer<float>, InputBias)
+			SHADER_PARAMETER_RDG_BUFFER_SRV(Buffer<float>, InputMean)
+			SHADER_PARAMETER_RDG_BUFFER_SRV(Buffer<float>, InputInvStdDev)
 			SHADER_PARAMETER_RDG_BUFFER_UAV(RWBuffer<float>, Output)
-		END_SHADER_PARAMETER_STRUCT()
-
-		static void FillInParameters(float Epsilon, const NNE::Internal::FTensor& Input, FParameters& Parameters);
-
-		static FIntVector GetGroupCount(const FParameters& Parameters, EInstanceNormalizationAlgorithm Algorithm);
-		static EInstanceNormalizationAlgorithm GetAlgorithm(const FParameters& Parameters);
+			END_SHADER_PARAMETER_STRUCT()
 
-		static void LexFromString(EInstanceNormalizationAlgorithm& OutValue, const TCHAR* StringVal);
+			static void ModifyCompilationEnvironment(const FGlobalShaderPermutationParameters& InParameters, FShaderCompilerEnvironment& OutEnvironment);
 	};
 } // UE::NNEHlslShaders::Internal
\ No newline at end of file
diff --git a/Engine/Plugins/Experimental/NNERuntimeRDG/Source/NNEHlslShaders/Internal/NNEHlslShadersReduceCS.h b/Engine/Plugins/Experimental/NNERuntimeRDG/Source/NNEHlslShaders/Internal/NNEHlslShadersReduceCS.h
--- a/Engine/Plugins/Experimental/NNERuntimeRDG/Source/NNEHlslShaders/Internal/NNEHlslShadersReduceCS.h
+++ b/Engine/Plugins/Experimental/NNERuntimeRDG/Source/NNEHlslShaders/Internal/NNEHlslShadersReduceCS.h
@@ -23,6 +23,7 @@ namespace UE::NNEHlslShaders::Internal
 		Prod,
 		Sum,
 		SumExp,//Should not be used for multiple axis reduction
+		AverageInvStdDev,//Should not be used for multiple axis reduction
 		MAX
 	};
 
@@ -46,12 +47,14 @@ namespace UE::NNEHlslShaders::Internal
 			SHADER_PARAMETER(int32, NumElemBeforeAxis)
 			SHADER_PARAMETER(int32, AxisSize)
 			SHADER_PARAMETER(int32, NumElemAfterAxis)
+			SHADER_PARAMETER(float, Epsilon)
 			SHADER_PARAMETER_RDG_BUFFER_SRV(Buffer<float>, Input)
 			SHADER_PARAMETER_RDG_BUFFER_UAV(RWBuffer<float>, Output)
+			SHADER_PARAMETER_RDG_BUFFER_UAV(RWBuffer<float>, Output2)
 		END_SHADER_PARAMETER_STRUCT()
 
 		static void ModifyCompilationEnvironment(const FGlobalShaderPermutationParameters& InParameters, FShaderCompilerEnvironment& OutEnvironment);
 		static void FillInParameters(TConstArrayView<uint32> Shape, int32 Axis, FParameters* Parameters);
-		static void EnqueueRDG(FRDGBuilder& GraphBuilder, FParameters* Parameters, FRDGBufferRef Input, FRDGBufferRef Output, EReduceOperatorType OperatorType);
+		static void EnqueueRDG(FRDGBuilder& GraphBuilder, FParameters* Parameters, FRDGBufferRef Input, FRDGBufferRef Output, EReduceOperatorType OperatorType, FRDGBufferRef Output2 = {});
 	};
 } // UE::NNEHlslShaders::Internal
\ No newline at end of file
diff --git a/Engine/Plugins/Experimental/NNERuntimeRDG/Source/NNEHlslShaders/Private/NNEHlslShadersInstanceNormalizationCS.cpp b/Engine/Plugins/Experimental/NNERuntimeRDG/Source/NNEHlslShaders/Private/NNEHlslShadersInstanceNormalizationCS.cpp
--- a/Engine/Plugins/Experimental/NNERuntimeRDG/Source/NNEHlslShaders/Private/NNEHlslShadersInstanceNormalizationCS.cpp
+++ b/Engine/Plugins/Experimental/NNERuntimeRDG/Source/NNEHlslShaders/Private/NNEHlslShadersInstanceNormalizationCS.cpp
@@ -6,57 +6,10 @@
 
 namespace UE::NNEHlslShaders::Internal
 {
-	namespace InstNormUtils
+	void TInstanceNormalizationCS::ModifyCompilationEnvironment(const FGlobalShaderPermutationParameters& InParameters, FShaderCompilerEnvironment& OutEnvironment)
 	{
-		FIntVector GetGroupSize(EInstanceNormalizationAlgorithm Algorithm)
-		{
-			switch (Algorithm)
-			{
-			case EInstanceNormalizationAlgorithm::Simple1x265: 			return {    1, 256, 1 };
-			case EInstanceNormalizationAlgorithm::SharedMemory256x1:	return {  256,   1, 1 };
-			case EInstanceNormalizationAlgorithm::SharedMemory512x1:	return {  512,   1, 1 };
-
-			// https://learn.microsoft.com/en-us/windows/win32/direct3dhlsl/sm5-attributes-numthreads
-			case EInstanceNormalizationAlgorithm::SharedMemory768x1:	return {  768,   1, 1 }; // Compute shader version 4.X
-			case EInstanceNormalizationAlgorithm::SharedMemory1024x1:	return { 1024,   1, 1 }; // Compute shader version 5.0
-			}
-
-			check(false);
-			return {1, 256, 1};
-		}
-	}
-
-	void TInstanceNormalizationCS::FillInParameters(float Epsilon, const  UE::NNE::Internal::FTensor& Input, TInstanceNormalizationCS::FParameters& Parameters)
-	{
-		check(Input.GetShape().Rank() >= 3);
-		
-		Parameters.Epsilon = Epsilon;
-		const int32 N = Input.GetShape().GetData()[0];
-		Parameters.C = Input.GetShape().GetData()[1];
-		Parameters.NxC = N * Parameters.C;
-		Parameters.W = Input.GetShape().Volume() / (Parameters.NxC);
-	}
-
-	FIntVector TInstanceNormalizationCS::GetGroupCount(const TInstanceNormalizationCS::FParameters& Parameters, EInstanceNormalizationAlgorithm Algorithm)
-	{
-		int ThreadGroupCountValueY = FMath::DivideAndRoundUp((int32)Parameters.NxC, InstNormUtils::GetGroupSize(Algorithm).Y);
-
-		return {1, ThreadGroupCountValueY, 1};
-	}
-
-	EInstanceNormalizationAlgorithm TInstanceNormalizationCS::GetAlgorithm(const TInstanceNormalizationCS::FParameters& Parameters)
-	{
-		return EInstanceNormalizationAlgorithm::SharedMemory768x1; // note: not all platforms support higher number of threads (yet)
-	}
-
-	void TInstanceNormalizationCS::LexFromString(EInstanceNormalizationAlgorithm& OutValue, const TCHAR* StringVal)
-	{
-		OutValue = EInstanceNormalizationAlgorithm::MAX;
-		if (FCString::Stricmp(StringVal, TEXT("Simple1x265")) == 0) OutValue = EInstanceNormalizationAlgorithm::Simple1x265;
-		else if (FCString::Stricmp(StringVal, TEXT("SharedMemory256x1")) == 0) OutValue = EInstanceNormalizationAlgorithm::SharedMemory256x1;
-		else if (FCString::Stricmp(StringVal, TEXT("SharedMemory512x1")) == 0) OutValue = EInstanceNormalizationAlgorithm::SharedMemory512x1;
-		else if (FCString::Stricmp(StringVal, TEXT("SharedMemory768x1")) == 0) OutValue = EInstanceNormalizationAlgorithm::SharedMemory768x1;
-		else if (FCString::Stricmp(StringVal, TEXT("SharedMemory1024x1")) == 0) OutValue = EInstanceNormalizationAlgorithm::SharedMemory1024x1;
+		FGlobalShader::ModifyCompilationEnvironment(InParameters, OutEnvironment);
+		OutEnvironment.SetDefine(TEXT("THREADGROUP_SIZE_X"), FInstanceNormalizationConstants::NUM_GROUP_THREADS);
 	}
 
 	IMPLEMENT_GLOBAL_SHADER(TInstanceNormalizationCS, "/NNE/NNEHlslShadersInstanceNormalization.usf", "InstanceNormalization", SF_Compute);
diff --git a/Engine/Plugins/Experimental/NNERuntimeRDG/Source/NNEHlslShaders/Private/NNEHlslShadersReduceCS.cpp b/Engine/Plugins/Experimental/NNERuntimeRDG/Source/NNEHlslShaders/Private/NNEHlslShadersReduceCS.cpp
--- a/Engine/Plugins/Experimental/NNERuntimeRDG/Source/NNEHlslShaders/Private/NNEHlslShadersReduceCS.cpp
+++ b/Engine/Plugins/Experimental/NNERuntimeRDG/Source/NNEHlslShaders/Private/NNEHlslShadersReduceCS.cpp
@@ -32,14 +32,18 @@ namespace UE::NNEHlslShaders::Internal
 		Parameters->NumElemBeforeAxis = NumElemBeforeAxis;
 		Parameters->AxisSize = Shape[Axis];
 		Parameters->NumElemAfterAxis = NumElemAfterAxis;
+		Parameters->Epsilon = 0;
 	}
 
-	void TReduceCS::EnqueueRDG(FRDGBuilder& GraphBuilder, TReduceCS::FParameters* Parameters, FRDGBufferRef Input, FRDGBufferRef Output, EReduceOperatorType OperatorType)
+	void TReduceCS::EnqueueRDG(FRDGBuilder& GraphBuilder, TReduceCS::FParameters* Parameters, FRDGBufferRef Input, FRDGBufferRef Output, EReduceOperatorType OperatorType, FRDGBufferRef Output2)
 	{
 		check(Parameters);
 
 		Parameters->Input = GraphBuilder.CreateSRV(FRDGBufferSRVDesc(Input, PF_R32_FLOAT));
 		Parameters->Output = GraphBuilder.CreateUAV(FRDGBufferUAVDesc(Output, PF_R32_FLOAT));
+		if (Output2 != nullptr) {
+			Parameters->Output2 = GraphBuilder.CreateUAV(FRDGBufferUAVDesc(Output2, PF_R32_FLOAT));
+		}
 
 		TReduceCS::FPermutationDomain PermutationVector;
 		PermutationVector.Set<TReduceCS::FReduceType>(OperatorType);
diff --git a/Engine/Plugins/Experimental/NNERuntimeRDG/Source/NNEHlslShaders/Shaders/NNEHlslShadersInstanceNormalization.usf b/Engine/Plugins/Experimental/NNERuntimeRDG/Source/NNEHlslShaders/Shaders/NNEHlslShadersInstanceNormalization.usf
--- a/Engine/Plugins/Experimental/NNERuntimeRDG/Source/NNEHlslShaders/Shaders/NNEHlslShadersInstanceNormalization.usf
+++ b/Engine/Plugins/Experimental/NNERuntimeRDG/Source/NNEHlslShaders/Shaders/NNEHlslShadersInstanceNormalization.usf
@@ -1,173 +1,36 @@
 // Copyright Epic Games, Inc. All Rights Reserved.
-
+ 
 #include "/Engine/Public/Platform.ush"
-
+ 
 #define WORK_TYPE float
 #define BUFFER_TYPE float
 #define READ(x) x
 #define WRITE(x) x
-
-int C;							// Number of channels
-int NxC;						// Precalculated N x C
-int W; 							// D1 x D2 ... Dn
-
-BUFFER_TYPE Epsilon;			// Attribute epsilon to use to avoid division by zero
-Buffer<BUFFER_TYPE> Input;		// Input data tensor with dimensions (N x C x D1 x D2 ... Dn)
-Buffer<BUFFER_TYPE> Scale;		// Scale tensor of size C
-Buffer<BUFFER_TYPE> Bias;		// Bias tensor of size C
-RWBuffer<BUFFER_TYPE> Output;	// Output tensor of the same shape as input
-
-#if ALGORITHM == 0
-#define GROUP_SIZE_X 1
-#define GROUP_SIZE_Y 256
-
-#elif ALGORITHM == 1
-#define GROUP_SIZE_X 256
-#define GROUP_SIZE_Y 1
-
-#elif ALGORITHM == 2
-#define GROUP_SIZE_X 512
-#define GROUP_SIZE_Y 1
-
-#elif ALGORITHM == 3
-#define GROUP_SIZE_X 768
-#define GROUP_SIZE_Y 1
-
-#elif ALGORITHM == 4
-#define GROUP_SIZE_X 1024
-#define GROUP_SIZE_Y 1
-
-#endif
-
-#if GROUP_SIZE_X > 1
-
-groupshared WORK_TYPE SharedMemoryMean[GROUP_SIZE_X * GROUP_SIZE_Y];
-groupshared WORK_TYPE SharedMemoryMean2[GROUP_SIZE_X * GROUP_SIZE_Y];
-groupshared int SharedMemoryCount[GROUP_SIZE_X * GROUP_SIZE_Y];
-
-[numthreads(GROUP_SIZE_X, GROUP_SIZE_Y, 1)]
-void InstanceNormalization(in const uint3 DispatchThreadID : SV_DispatchThreadID, in const uint3 GroupThreadID : SV_GroupThreadID)
-{
-	const int StartIdx = DispatchThreadID.y * W;
-	const int SharedMemoryOffsetY = GroupThreadID.y * GROUP_SIZE_X;
-
-	// Calculate mean and variance using Welford's online algorithm
-	// https://en.wikipedia.org/wiki/Algorithms_for_calculating_variance
-	if (DispatchThreadID.y < NxC)
-	{
-		WORK_TYPE Mean = 0;
-		WORK_TYPE Mean2 = 0;
-		int Count = 0;
-		for (int CurrIdx = GroupThreadID.x; CurrIdx < W; CurrIdx += GROUP_SIZE_X)
-		{
-			Count += 1;
-
-			WORK_TYPE Value = READ(Input[StartIdx + CurrIdx]);
-
-			WORK_TYPE Delta = Value - Mean;
-			Mean += Delta / (WORK_TYPE)Count;
-
-			WORK_TYPE Delta2 = Value - Mean;
-			Mean2 += Delta * Delta2;
-		}
-
-		SharedMemoryMean[SharedMemoryOffsetY + GroupThreadID.x] = Mean;
-		SharedMemoryMean2[SharedMemoryOffsetY + GroupThreadID.x] = Mean2;
-		SharedMemoryCount[SharedMemoryOffsetY + GroupThreadID.x] = Count;
-	}
-
-	// Wait until all threads wrote result to shared memory
-	GroupMemoryBarrierWithGroupSync();
-
-	// Use one thread to calculate final result
-	if (GroupThreadID.x == 0 && DispatchThreadID.y < NxC)
-	{
-		WORK_TYPE Mean = 0;
-		WORK_TYPE Mean2 = 0;
-		int Count = 0;
-
-		UNROLL
-		for (int Idx = 0; Idx < GROUP_SIZE_X; Idx++)
-		{
-			WORK_TYPE CurrMean = SharedMemoryMean[SharedMemoryOffsetY + Idx];
-			WORK_TYPE CurrMean2 = SharedMemoryMean2[SharedMemoryOffsetY + Idx];
-			int CurrCount = SharedMemoryCount[SharedMemoryOffsetY + Idx];
-
-			int NewCount = Count + CurrCount;
-
-			WORK_TYPE Delta = CurrMean - Mean;
-			Mean = (CurrCount * CurrMean + Count * Mean) / NewCount;
-
-			Mean2 += CurrMean2 + Delta * Delta * (WORK_TYPE)Count * (WORK_TYPE)CurrCount / (WORK_TYPE)NewCount;
-
-			Count = NewCount;
-		}
-
-		WORK_TYPE Variance = Mean2 / (WORK_TYPE)Count;
-
-		SharedMemoryMean[SharedMemoryOffsetY] = Mean;
-		SharedMemoryMean2[SharedMemoryOffsetY] = (WORK_TYPE)1 / sqrt(Variance + Epsilon);
-	}
-
-	// Wait until first thread wrote result to shared memory
-	GroupMemoryBarrierWithGroupSync();
-
-	if (DispatchThreadID.y < NxC)
-	{
-		const int ChannelIdx = DispatchThreadID.y % C;
-		
-		WORK_TYPE Mean = SharedMemoryMean[SharedMemoryOffsetY];
-		WORK_TYPE InvStdev = SharedMemoryMean2[SharedMemoryOffsetY];
-
-		// Carry out instance normalization and write to output
-		for (int CurrIdx = GroupThreadID.x; CurrIdx < W; CurrIdx += GROUP_SIZE_X)
-		{
-			Output[StartIdx + CurrIdx] = WRITE((READ(Input[StartIdx + CurrIdx]) - Mean) * READ(Scale[ChannelIdx]) * InvStdev + READ(Bias[ChannelIdx]));
-		}
-	}
-}
-
-#else
-
-[numthreads(1, 256, 1)]
+ 
+Buffer<WORK_TYPE> Input;
+Buffer<WORK_TYPE> InputScale;
+Buffer<WORK_TYPE> InputBias;
+Buffer<WORK_TYPE> InputMean;
+Buffer<WORK_TYPE> InputInvStdDev;
+RWBuffer<WORK_TYPE> Output;
+uint Num;
+uint ThreadCountX;
+uint InstanceSize;
+uint ChannelSize;
+ 
+[numthreads(THREADGROUP_SIZE_X, 1, 1)]
 void InstanceNormalization(in const uint3 DispatchThreadID : SV_DispatchThreadID)
 {
-	if (DispatchThreadID.y >= NxC || DispatchThreadID.x >= 1)
-	{
-		return;
-	}
-
-	const int StartIdx = DispatchThreadID.y * W;
-	const int ChannelIdx = DispatchThreadID.y % C;
-
-	const WORK_TYPE InvNum = (WORK_TYPE)1 / (WORK_TYPE)W;
+	const uint Index = DispatchThreadID.y * ThreadCountX + DispatchThreadID.x;
 
-	// Calculate mean and variance using Welfor's online algorithm
-	WORK_TYPE Mean = 0;
-	WORK_TYPE Mean2 = 0;
-	int Count = 0;
-	for (int Idx = 0; Idx < W; Idx++)
+	if (Index < Num)
 	{
-		Count += 1;
-
-		WORK_TYPE Value = READ(Input[StartIdx + Idx]);
-
-		WORK_TYPE Delta = Value - Mean;
-		Mean += Delta / (WORK_TYPE)Count;
-
-		WORK_TYPE Delta2 = Value - Mean;
-		Mean2 += Delta * Delta2;
+		const uint InstanceIdx = Index/InstanceSize;
+		const uint ChannelIdx = InstanceIdx%ChannelSize;
+		WORK_TYPE Mean = READ(InputMean[InstanceIdx]);
+		WORK_TYPE InvStdDev = READ(InputInvStdDev[InstanceIdx]);
+		WORK_TYPE Scale = READ(InputScale[ChannelIdx]);
+		WORK_TYPE Bias = READ(InputBias[ChannelIdx]);
+		Output[Index] = WRITE((READ(Input[Index]) - Mean) * Scale * InvStdDev + Bias);
 	}
-
-	WORK_TYPE Variance = Mean2 * InvNum;
-
-	WORK_TYPE InvStdev = (WORK_TYPE)1 / sqrt(Variance + Epsilon);
-
-	// Carry out instance normalization and write to output
-	for (int Idx = 0; Idx < W; Idx++)
-	{
-		Output[StartIdx + Idx] = WRITE((READ(Input[StartIdx + Idx]) - Mean) * READ(Scale[ChannelIdx]) * InvStdev + READ(Bias[ChannelIdx]));
-	}
-}
-
-#endif
\ No newline at end of file
+}
\ No newline at end of file
diff --git a/Engine/Plugins/Experimental/NNERuntimeRDG/Source/NNEHlslShaders/Shaders/NNEHlslShadersReduce.usf b/Engine/Plugins/Experimental/NNERuntimeRDG/Source/NNEHlslShaders/Shaders/NNEHlslShadersReduce.usf
--- a/Engine/Plugins/Experimental/NNERuntimeRDG/Source/NNEHlslShaders/Shaders/NNEHlslShadersReduce.usf
+++ b/Engine/Plugins/Experimental/NNERuntimeRDG/Source/NNEHlslShaders/Shaders/NNEHlslShadersReduce.usf
@@ -9,9 +9,11 @@
 
 Buffer<WORK_TYPE> Input;
 RWBuffer<WORK_TYPE> Output;
+RWBuffer<BUFFER_TYPE> Output2;
 int NumElemBeforeAxis;
 int AxisSize;
 int NumElemAfterAxis;
+float Epsilon;
 
 // Must correspond to EReduceOperatorType defined in NNEHlslShadersReduceCS.h
 #define TYPE_AVERAGE 0
@@ -23,100 +25,107 @@ int NumElemAfterAxis;
 #define TYPE_PROD 6
 #define TYPE_SUM 7
 #define TYPE_SUMEXP 8
+#define TYPE_AVERAGEINVSTDDEV 9
 
-WORK_TYPE InitValue(WORK_TYPE FirstValue)
-{
-	#if REDUCE_OPERATOR_TYPE == TYPE_AVERAGE
-		return 0;
-	#elif REDUCE_OPERATOR_TYPE == TYPE_L1
-		return 0;
-	#elif REDUCE_OPERATOR_TYPE == TYPE_L2
-		return 0;
-	#elif REDUCE_OPERATOR_TYPE == TYPE_LOGSUMEXP
-		return 0;
-	#elif REDUCE_OPERATOR_TYPE == TYPE_MAX
-		return FirstValue;
-	#elif REDUCE_OPERATOR_TYPE == TYPE_MIN
-		return FirstValue;
-	#elif REDUCE_OPERATOR_TYPE == TYPE_PROD
-		return 1;
-	#elif REDUCE_OPERATOR_TYPE == TYPE_SUM
-		return 0;
-	#elif REDUCE_OPERATOR_TYPE == TYPE_SUMEXP
-		return 0;
-	#endif
-}
-
-WORK_TYPE ElementOp(WORK_TYPE Value)
-{
-	#if REDUCE_OPERATOR_TYPE == TYPE_AVERAGE
-		return Value;
-	#elif REDUCE_OPERATOR_TYPE == TYPE_L1
-		return abs(Value);
-	#elif REDUCE_OPERATOR_TYPE == TYPE_L2
-		return Value*Value;
-	#elif REDUCE_OPERATOR_TYPE == TYPE_LOGSUMEXP
-		return exp(Value);
-	#elif REDUCE_OPERATOR_TYPE == TYPE_MAX
-		return Value;
-	#elif REDUCE_OPERATOR_TYPE == TYPE_MIN
-		return Value;
-	#elif REDUCE_OPERATOR_TYPE == TYPE_PROD
-		return Value;
-	#elif REDUCE_OPERATOR_TYPE == TYPE_SUM
-		return Value;
-	#elif REDUCE_OPERATOR_TYPE == TYPE_SUMEXP
-		return exp(Value);
-	#endif
-}
-
-WORK_TYPE CombineOp(WORK_TYPE ValueA, WORK_TYPE ValueB)
-{
-	#if REDUCE_OPERATOR_TYPE == TYPE_AVERAGE
-		return ValueA+ValueB;
-	#elif REDUCE_OPERATOR_TYPE == TYPE_L1
-		return ValueA+ValueB;
-	#elif REDUCE_OPERATOR_TYPE == TYPE_L2
-		return ValueA+ValueB;
-	#elif REDUCE_OPERATOR_TYPE == TYPE_LOGSUMEXP
-		return ValueA+ValueB;
-	#elif REDUCE_OPERATOR_TYPE == TYPE_MAX
-		return max(ValueA, ValueB);
-	#elif REDUCE_OPERATOR_TYPE == TYPE_MIN
-		return min(ValueA, ValueB);
-	#elif REDUCE_OPERATOR_TYPE == TYPE_PROD
-		return ValueA*ValueB;
-	#elif REDUCE_OPERATOR_TYPE == TYPE_SUM
-		return ValueA+ValueB;
-	#elif REDUCE_OPERATOR_TYPE == TYPE_SUMEXP
-		return ValueA+ValueB;
-	#endif
-}
-
-WORK_TYPE FinalOp(WORK_TYPE Value)
-{
-	#if REDUCE_OPERATOR_TYPE == TYPE_AVERAGE
-		return Value/AxisSize;
-	#elif REDUCE_OPERATOR_TYPE == TYPE_L1
-		return Value;
-	#elif REDUCE_OPERATOR_TYPE == TYPE_L2
-		return sqrt(Value);
-	#elif REDUCE_OPERATOR_TYPE == TYPE_LOGSUMEXP
-		return log(Value);
-	#elif REDUCE_OPERATOR_TYPE == TYPE_MAX
-		return Value;
-	#elif REDUCE_OPERATOR_TYPE == TYPE_MIN
-		return Value;
-	#elif REDUCE_OPERATOR_TYPE == TYPE_PROD
-		return Value;
-	#elif REDUCE_OPERATOR_TYPE == TYPE_SUM
-		return Value;
-	#elif REDUCE_OPERATOR_TYPE == TYPE_SUMEXP
-		return Value;
-	#endif
-}
-
-groupshared WORK_TYPE SharedMemory[THREADGROUP_SIZE];
+#if REDUCE_OPERATOR_TYPE == TYPE_AVERAGEINVSTDDEV
+	groupshared WORK_TYPE SharedMemoryMean[THREADGROUP_SIZE];
+	groupshared WORK_TYPE SharedMemoryMean2[THREADGROUP_SIZE];
+	groupshared WORK_TYPE SharedMemoryCount[THREADGROUP_SIZE];
+#else //REDUCE_OPERATOR_TYPE != TYPE_AVERAGEINVSTD
+	groupshared WORK_TYPE SharedMemory[THREADGROUP_SIZE];
+
+	WORK_TYPE InitValue(WORK_TYPE FirstValue)
+	{
+		#if REDUCE_OPERATOR_TYPE == TYPE_AVERAGE
+			return 0;
+		#elif REDUCE_OPERATOR_TYPE == TYPE_L1
+			return 0;
+		#elif REDUCE_OPERATOR_TYPE == TYPE_L2
+			return 0;
+		#elif REDUCE_OPERATOR_TYPE == TYPE_LOGSUMEXP
+			return 0;
+		#elif REDUCE_OPERATOR_TYPE == TYPE_MAX
+			return FirstValue;
+		#elif REDUCE_OPERATOR_TYPE == TYPE_MIN
+			return FirstValue;
+		#elif REDUCE_OPERATOR_TYPE == TYPE_PROD
+			return 1;
+		#elif REDUCE_OPERATOR_TYPE == TYPE_SUM
+			return 0;
+		#elif REDUCE_OPERATOR_TYPE == TYPE_SUMEXP
+			return 0;
+		#endif
+	}
+
+	WORK_TYPE ElementOp(WORK_TYPE Value)
+	{
+		#if REDUCE_OPERATOR_TYPE == TYPE_AVERAGE
+			return Value;
+		#elif REDUCE_OPERATOR_TYPE == TYPE_L1
+			return abs(Value);
+		#elif REDUCE_OPERATOR_TYPE == TYPE_L2
+			return Value*Value;
+		#elif REDUCE_OPERATOR_TYPE == TYPE_LOGSUMEXP
+			return exp(Value);
+		#elif REDUCE_OPERATOR_TYPE == TYPE_MAX
+			return Value;
+		#elif REDUCE_OPERATOR_TYPE == TYPE_MIN
+			return Value;
+		#elif REDUCE_OPERATOR_TYPE == TYPE_PROD
+			return Value;
+		#elif REDUCE_OPERATOR_TYPE == TYPE_SUM
+			return Value;
+		#elif REDUCE_OPERATOR_TYPE == TYPE_SUMEXP
+			return exp(Value);
+		#endif
+	}
+
+	WORK_TYPE CombineOp(WORK_TYPE ValueA, WORK_TYPE ValueB)
+	{
+		#if REDUCE_OPERATOR_TYPE == TYPE_AVERAGE
+			return ValueA+ValueB;
+		#elif REDUCE_OPERATOR_TYPE == TYPE_L1
+			return ValueA+ValueB;
+		#elif REDUCE_OPERATOR_TYPE == TYPE_L2
+			return ValueA+ValueB;
+		#elif REDUCE_OPERATOR_TYPE == TYPE_LOGSUMEXP
+			return ValueA+ValueB;
+		#elif REDUCE_OPERATOR_TYPE == TYPE_MAX
+			return max(ValueA, ValueB);
+		#elif REDUCE_OPERATOR_TYPE == TYPE_MIN
+			return min(ValueA, ValueB);
+		#elif REDUCE_OPERATOR_TYPE == TYPE_PROD
+			return ValueA*ValueB;
+		#elif REDUCE_OPERATOR_TYPE == TYPE_SUM
+			return ValueA+ValueB;
+		#elif REDUCE_OPERATOR_TYPE == TYPE_SUMEXP
+			return ValueA+ValueB;
+		#endif
+	}
+
+	WORK_TYPE FinalOp(WORK_TYPE Value)
+	{
+		#if REDUCE_OPERATOR_TYPE == TYPE_AVERAGE
+			return Value/AxisSize;
+		#elif REDUCE_OPERATOR_TYPE == TYPE_L1
+			return Value;
+		#elif REDUCE_OPERATOR_TYPE == TYPE_L2
+			return sqrt(Value);
+		#elif REDUCE_OPERATOR_TYPE == TYPE_LOGSUMEXP
+			return log(Value);
+		#elif REDUCE_OPERATOR_TYPE == TYPE_MAX
+			return Value;
+		#elif REDUCE_OPERATOR_TYPE == TYPE_MIN
+			return Value;
+		#elif REDUCE_OPERATOR_TYPE == TYPE_PROD
+			return Value;
+		#elif REDUCE_OPERATOR_TYPE == TYPE_SUM
+			return Value;
+		#elif REDUCE_OPERATOR_TYPE == TYPE_SUMEXP
+			return Value;
+		#endif
+	}
+#endif //REDUCE_OPERATOR_TYPE != TYPE_AVERAGEINVSTD
 
 [numthreads(THREADGROUP_SIZE, 1, 1)]
 void Reduce(in const uint3 DispatchThreadID : SV_DispatchThreadID, in const uint3 GroupThreadID : SV_GroupThreadID)
@@ -128,18 +137,41 @@ void Reduce(in const uint3 DispatchThreadID : SV_DispatchThreadID, in const uint
 
 	// Reduction using sharedmem and THREADGROUP_SIZE threads along the reduced dimension
 	{
-		WORK_TYPE FirstValue = READ(Input[InputStartIdx]);
-		WORK_TYPE Reduced = InitValue(FirstValue);
-
-		for (int AxisIdx = GroupThreadID.x; AxisIdx < AxisSize; AxisIdx += THREADGROUP_SIZE)
-		{
-			WORK_TYPE Value = READ(Input[(AxisIdx * NumElemAfterAxis) + InputStartIdx]);
-
-			Value = ElementOp(Value);
-			Reduced = CombineOp(Reduced, Value);
-		}
-
-		SharedMemory[GroupThreadID.x] = Reduced;
+		#if REDUCE_OPERATOR_TYPE == TYPE_AVERAGEINVSTDDEV
+			WORK_TYPE Mean = 0;
+			WORK_TYPE Mean2 = 0;
+			int Count = 0;
+		#else //REDUCE_OPERATOR_TYPE != TYPE_AVERAGEINVSTDDEV
+			WORK_TYPE FirstValue = READ(Input[InputStartIdx]);
+			WORK_TYPE Reduced = InitValue(FirstValue);
+		#endif
+
+
+			for (int AxisIdx = GroupThreadID.x; AxisIdx < AxisSize; AxisIdx += THREADGROUP_SIZE)
+			{
+				WORK_TYPE Value = READ(Input[(AxisIdx * NumElemAfterAxis) + InputStartIdx]);
+			
+				#if REDUCE_OPERATOR_TYPE == TYPE_AVERAGEINVSTDDEV
+					// Calculate mean and variance using Welford's online algorithm
+					// https://en.wikipedia.org/wiki/Algorithms_for_calculating_variance
+					Count += 1;
+					WORK_TYPE Delta = Value - Mean;
+					Mean += Delta / (WORK_TYPE)Count;
+					WORK_TYPE Delta2 = Value - Mean;
+					Mean2 += Delta * Delta2;
+				#else //REDUCE_OPERATOR_TYPE != TYPE_AVERAGEINVSTDDEV
+					Value = ElementOp(Value);
+					Reduced = CombineOp(Reduced, Value);
+				#endif
+			}
+
+		#if REDUCE_OPERATOR_TYPE == TYPE_AVERAGEINVSTDDEV
+			SharedMemoryMean[GroupThreadID.x] = Mean;
+			SharedMemoryMean2[GroupThreadID.x] = Mean2;
+			SharedMemoryCount[GroupThreadID.x] = Count;
+		#else //REDUCE_OPERATOR_TYPE != TYPE_AVERAGEINVSTDDEV
+			SharedMemory[GroupThreadID.x] = Reduced;
+		#endif
 	}
 
 	// Wait until all threads wrote result to shared memory
@@ -148,25 +180,42 @@ void Reduce(in const uint3 DispatchThreadID : SV_DispatchThreadID, in const uint
 	// Use one thread to calculate final result
 	if (GroupThreadID.x == 0)
 	{
-		WORK_TYPE FirstCurrReduced = SharedMemory[0];
-		WORK_TYPE Reduced = InitValue(FirstCurrReduced);
-
-		UNROLL
-		for (int Idx = 0; Idx < THREADGROUP_SIZE; Idx++)
-		{
-			WORK_TYPE CurrReduced = SharedMemory[Idx];
-
-			Reduced = CombineOp(Reduced, CurrReduced);
-		}
-
-		SharedMemory[0] = Reduced;
+		#if REDUCE_OPERATOR_TYPE == TYPE_AVERAGEINVSTDDEV
+			WORK_TYPE Mean = 0;
+			WORK_TYPE Mean2 = 0;
+			int Count = 0;
+		#else //REDUCE_OPERATOR_TYPE != TYPE_AVERAGEINVSTDDEV
+			WORK_TYPE FirstCurrReduced = SharedMemory[0];
+			WORK_TYPE Reduced = InitValue(FirstCurrReduced);
+		#endif
+
+			UNROLL
+			for (int Idx = 0; Idx < THREADGROUP_SIZE; Idx++)
+			{
+				#if REDUCE_OPERATOR_TYPE == TYPE_AVERAGEINVSTDDEV
+					WORK_TYPE CurrMean = SharedMemoryMean[Idx];
+					WORK_TYPE CurrMean2 = SharedMemoryMean2[Idx];
+					int CurrCount = SharedMemoryCount[Idx];
+					int NewCount = Count + CurrCount;
+					WORK_TYPE Delta = CurrMean - Mean;
+					Mean = (CurrCount * CurrMean + Count * Mean) / NewCount;
+					Mean2 += CurrMean2 + Delta * Delta * (WORK_TYPE)Count * (WORK_TYPE)CurrCount / (WORK_TYPE)NewCount;
+					Count = NewCount;
+				#else //REDUCE_OPERATOR_TYPE != TYPE_AVERAGEINVSTDDEV
+					WORK_TYPE CurrReduced = SharedMemory[Idx];
+					Reduced = CombineOp(Reduced, CurrReduced);
+				#endif
+			}
+
+		// Write reduced result
+		#if REDUCE_OPERATOR_TYPE == TYPE_AVERAGEINVSTDDEV
+			WORK_TYPE Variance = Mean2 / (WORK_TYPE)Count;
+			WORK_TYPE InvStdev = (WORK_TYPE)1 / sqrt(Variance + Epsilon);
+			Output[OutputIdx] = Mean;
+			Output2[OutputIdx] = InvStdev;
+		#else //REDUCE_OPERATOR_TYPE != TYPE_AVERAGEINVSTDDEV
+			Reduced = FinalOp(Reduced);
+			Output[OutputIdx] = Reduced;
+		#endif
 	}
-
-	// Wait until first thread wrote result to shared memory
-	GroupMemoryBarrierWithGroupSync();
-
-	// Write reduced result
-	WORK_TYPE Reduced = SharedMemory[0];
-	Reduced = FinalOp(Reduced);
-	Output[OutputIdx] = Reduced;
 }
\ No newline at end of file
diff --git a/Engine/Plugins/Experimental/NNERuntimeRDG/Source/NNERuntimeRDG/Private/Hlsl/NNERuntimeRDGInstanceNormalization.cpp b/Engine/Plugins/Experimental/NNERuntimeRDG/Source/NNERuntimeRDG/Private/Hlsl/NNERuntimeRDGInstanceNormalization.cpp
--- a/Engine/Plugins/Experimental/NNERuntimeRDG/Source/NNERuntimeRDG/Private/Hlsl/NNERuntimeRDGInstanceNormalization.cpp
+++ b/Engine/Plugins/Experimental/NNERuntimeRDG/Source/NNERuntimeRDG/Private/Hlsl/NNERuntimeRDGInstanceNormalization.cpp
@@ -2,6 +2,7 @@
 
 #include "NNERuntimeRDGInstanceNormalization.h"
 #include "NNEHlslShadersInstanceNormalizationCS.h"
+#include "NNEHlslShadersReduceCS.h"
 #include "NNERuntimeRDGHlslHelper.h"
 #include "NNETensor.h"
 #include "NNETypes.h"
@@ -11,12 +12,6 @@ namespace
 
 int ValidateInput(const UE::NNE::FTensorShape& Input, const UE::NNE::FTensorShape& Scale, const UE::NNE::FTensorShape& Bias)
 {
-	if (Input.Rank() < 3)
-	{
-		UE_LOG(LogNNE, Warning, TEXT("InstanceNormalization input data should be at least of rank 4: %d"), Input.Rank());
-		return -1;
-	}
-
 	if (Scale.Rank() != 1)
 	{
 		UE_LOG(LogNNE, Warning, TEXT("InstanceNormalization input scale should be of rank 1: %d"), Scale.Rank());
@@ -28,14 +23,14 @@ int ValidateInput(const UE::NNE::FTensorShape& Input, const UE::NNE::FTensorShap
 		return -1;
 	}
 
-	if (Bias.GetData()[0] != Input.GetData()[1])
+	if (Bias.Rank() != 1)
 	{
-		UE_LOG(LogNNE, Warning, TEXT("InstanceNormalization intput B size should be equal to channel count: : %d vs %d"), Bias.GetData()[0], Input.GetData()[1]);
+		UE_LOG(LogNNE, Warning, TEXT("InstanceNormalization input B should be of rank 1: %d"), Bias.Rank());
 		return -1;
 	}
-	if (Bias.Rank() != 1)
+	if (Bias.GetData()[0] != Input.GetData()[1])
 	{
-		UE_LOG(LogNNE, Warning, TEXT("InstanceNormalization input B should be of rank 1: %d"), Bias.Rank());
+		UE_LOG(LogNNE, Warning, TEXT("InstanceNormalization intput B size should be equal to channel count: : %d vs %d"), Bias.GetData()[0], Input.GetData()[1]);
 		return -1;
 	}
 
@@ -48,8 +43,6 @@ namespace UE::NNERuntimeRDG::Private::Hlsl
 {
 	DECLARE_GPU_STAT_NAMED(FNNEOperatorInstanceNormalization, TEXT("NNE.Operator.Hlsl.InstanceNormalization"));
 
-	using EInstanceNormalizationAlgorithm = UE::NNEHlslShaders::Internal::EInstanceNormalizationAlgorithm;
-
 	/**
 	 * InstanceNormalization operator implementation
 	 */
@@ -64,14 +57,13 @@ namespace UE::NNERuntimeRDG::Private::Hlsl
 
 		float Epsilon = 1e-5;
 
-		EInstanceNormalizationAlgorithm Algorithm = EInstanceNormalizationAlgorithm::MAX;
-
 	public:
 
 		virtual int PrepareOutputs(TConstArrayView<NNE::Internal::FTensorRef> InputTensors, TArrayView<NNE::Internal::FTensorRef> OutputTensors) const override
 		{
 			check(InputTensors.Num() == 3);
 			check(OutputTensors.Num() == 1);
+			check(InputTensors[0]->GetShape().Rank() > 2);
 
 			if (const int Res = ValidateInput(InputTensors[0]->GetShape(), InputTensors[1]->GetShape(), InputTensors[2]->GetShape()); Res < 0)
 			{
@@ -92,8 +84,19 @@ namespace UE::NNERuntimeRDG::Private::Hlsl
 
 			Epsilon = Attributes.GetValueOrDefault<float>(TEXT("epsilon"), 1e-5f);
 
-			// For testing only
-			TInstanceNormalizationCS::LexFromString(Algorithm, *Attributes.GetValueOrDefault<FString>(TEXT("__UE__algorithm"), "MAX"));
+			const int32 InputRank = InputTensorDescs[0].GetShape().Rank();
+
+			if (InputRank < 3)
+			{
+				UE_LOG(LogNNE, Warning, TEXT("InstanceNormalization input data should be at least of rank 3 but got: %d"), InputRank);
+				return false;
+			}
+
+			if (InputRank != OutputTensorDescs[0].GetShape().Rank())
+			{
+				UE_LOG(LogNNE, Warning, TEXT("InstanceNormalization requires the output to have the same rank as the input."));
+				return false;
+			}
 
 			return true;
 		}
@@ -113,35 +116,50 @@ namespace UE::NNERuntimeRDG::Private::Hlsl
 			const FTensorRDG& Scale = *InputTensors[1];
 			const FTensorRDG& Bias = *InputTensors[2];
 			const FTensorRDG& Output = *OutputTensors[0];
-
-			// Set parameters
-			TInstanceNormalizationCS::FParameters* Parameters = GraphBuilder.AllocParameters<TInstanceNormalizationCS::FParameters>();
-			TInstanceNormalizationCS::FillInParameters(Epsilon, Input, *Parameters);
-			Parameters->Input = GraphBuilder.CreateSRV(FRDGBufferSRVDesc(Input.GetBuffer(), PF_R32_FLOAT));
-			Parameters->Scale = GraphBuilder.CreateSRV(FRDGBufferSRVDesc(Scale.GetBuffer(), PF_R32_FLOAT));
-			Parameters->Bias = GraphBuilder.CreateSRV(FRDGBufferSRVDesc(Bias.GetBuffer(), PF_R32_FLOAT));
-			Parameters->Output = GraphBuilder.CreateUAV(FRDGBufferUAVDesc(Output.GetBuffer(), PF_R32_FLOAT));
-
-			EInstanceNormalizationAlgorithm DispatchAlgorithm = Algorithm;
-			if (DispatchAlgorithm == EInstanceNormalizationAlgorithm::MAX)
-			{
-				DispatchAlgorithm = TInstanceNormalizationCS::GetAlgorithm(*Parameters);
-			}
-			
-			TInstanceNormalizationCS::FPermutationDomain PermutationVector;
-			PermutationVector.Set<TInstanceNormalizationCS::FInstanceNormalizationAlgorithm>(DispatchAlgorithm);
-			TShaderMapRef<TInstanceNormalizationCS> ComputeShader(GetGlobalShaderMap(GMaxRHIFeatureLevel), PermutationVector);
-			FIntVector ThreadGroupCount = TInstanceNormalizationCS::GetGroupCount(*Parameters, DispatchAlgorithm);
+			const NNE::FTensorShape& InputShape = Input.GetShape();
+			constexpr int32 ReductionAxis = 2;
+			check(InputShape.Rank() >= ReductionAxis);
 
 			RDG_EVENT_SCOPE(GraphBuilder, "NNE.Operator.Hlsl.InstanceNormalization");
 			RDG_GPU_STAT_SCOPE(GraphBuilder, FNNEOperatorInstanceNormalization);
 
+			// First apply Reduction() to temp buffers getting Mean and InvStdDev
+			TReduceCS::FParameters* ReduceParameters = GraphBuilder.AllocParameters<TReduceCS::FParameters>();
+			TReduceCS::FillInParameters(InputShape.GetData(), ReductionAxis, ReduceParameters);
+			ReduceParameters->AxisSize *= ReduceParameters->NumElemAfterAxis;// InstanceNormalization flatten the input tensor to a 2D one
+			ReduceParameters->NumElemAfterAxis = 1;
+			ReduceParameters->Epsilon = Epsilon;
+			const FRDGBufferDesc InstNormTempBufferDesc = FRDGBufferDesc::CreateBufferDesc(Output.GetElementByteSize(), ReduceParameters->NumElemBeforeAxis);
+
+			FRDGBufferRef MeanBuffer = GraphBuilder.CreateBuffer(InstNormTempBufferDesc, TEXT("NNE.Operator.Hlsl.InstanceNormalization.TempMeanBuffer"), ERDGBufferFlags::None);
+			FRDGBufferRef InvStdDevBuffer = GraphBuilder.CreateBuffer(InstNormTempBufferDesc, TEXT("NNE.Operator.Hlsl.InstanceNormalization.TempInvStdDevBuffer"), ERDGBufferFlags::None);
+
+			TReduceCS::EnqueueRDG(GraphBuilder, ReduceParameters, Input.GetBuffer(), MeanBuffer, EReduceOperatorType::AverageInvStdDev, InvStdDevBuffer);
+
+			//Then InstanceNormalization
+			const int32 NumElements = Input.GetVolume();
+			const FIntVector ThreadGroupCount = ComputeElementWiseThreadGroups(NumElements, FInstanceNormalizationConstants::NUM_GROUP_THREADS);
+
+			TInstanceNormalizationCS::FParameters* InstNormParameters = GraphBuilder.AllocParameters<TInstanceNormalizationCS::FParameters>();
+			InstNormParameters->InstanceSize = ReduceParameters->AxisSize;
+			InstNormParameters->ChannelSize = InputShape.GetData()[1];
+			InstNormParameters->Num = NumElements;
+			InstNormParameters->ThreadCountX = ThreadGroupCount.X * FInstanceNormalizationConstants::NUM_GROUP_THREADS;
+			InstNormParameters->Input = GraphBuilder.CreateSRV(FRDGBufferSRVDesc(Input.GetBuffer(), PF_R32_FLOAT));
+			InstNormParameters->InputScale = GraphBuilder.CreateSRV(FRDGBufferSRVDesc(Scale.GetBuffer(), PF_R32_FLOAT));
+			InstNormParameters->InputBias = GraphBuilder.CreateSRV(FRDGBufferSRVDesc(Bias.GetBuffer(), PF_R32_FLOAT));
+			InstNormParameters->InputMean = GraphBuilder.CreateSRV(FRDGBufferSRVDesc(MeanBuffer, PF_R32_FLOAT));
+			InstNormParameters->InputInvStdDev = GraphBuilder.CreateSRV(FRDGBufferSRVDesc(InvStdDevBuffer, PF_R32_FLOAT));
+			InstNormParameters->Output = GraphBuilder.CreateUAV(FRDGBufferUAVDesc(Output.GetBuffer(), PF_R32_FLOAT));
+
+			TShaderMapRef<TInstanceNormalizationCS> ComputeShader(GetGlobalShaderMap(GMaxRHIFeatureLevel));
+
 			FComputeShaderUtils::AddPass(
 				GraphBuilder,
 				RDG_EVENT_NAME("NNE.Operator.Hlsl.InstanceNormalization.Dispatch"),
 				ERDGPassFlags::Compute | ERDGPassFlags::NeverCull,
 				ComputeShader,
-				Parameters,
+				InstNormParameters,
 				ThreadGroupCount);
 		}
 	};
@@ -152,7 +170,6 @@ namespace UE::NNERuntimeRDG::Private::Hlsl
 
 		FAttributeValidator AttributeValidator;
 		AttributeValidator.AddOptional(TEXT("epsilon"), ENNEAttributeDataType::Float);
-		AttributeValidator.AddOptional(TEXT("__UE__algorithm"), ENNEAttributeDataType::String);
 		bIsValid &= AttributeValidator.Validate(AttributeMap);
 
 		FInputValidator InputValidator;
