[NNE] Dropout operator version 7 support for HLSL.

Tests:
 * Non-Unity build
 * UnrealEditor, Win64: Dropout tests, smoke tests
 * UnrealGame, Win64: all tests

#rnx

#virtualized

diff --git a/Engine/Plugins/Experimental/NNERuntimeRDG/Source/NNERuntimeRDG/Private/Hlsl/NNERuntimeRDGDropout.cpp b/Engine/Plugins/Experimental/NNERuntimeRDG/Source/NNERuntimeRDG/Private/Hlsl/NNERuntimeRDGDropout.cpp
--- a/Engine/Plugins/Experimental/NNERuntimeRDG/Source/NNERuntimeRDG/Private/Hlsl/NNERuntimeRDGDropout.cpp
+++ b/Engine/Plugins/Experimental/NNERuntimeRDG/Source/NNERuntimeRDG/Private/Hlsl/NNERuntimeRDGDropout.cpp
@@ -68,12 +68,21 @@ namespace UE::NNERuntimeRDG::Private::Hlsl
 		}
 	};
 
+	template<int Version>
 	bool ValidateDropoutOperator(const NNE::FAttributeMap& AttributeMap, TConstArrayView<ENNETensorDataType> InputTypes, TConstArrayView<NNE::FSymbolicTensorShape> InputDropouts)
 	{
 		bool bIsValid = true;
 
 		FAttributeValidator AttributeValidator;
-		AttributeValidator.AddOptional(TEXT("seed"), ENNEAttributeDataType::Int32); //Will be ignored, only useful in training mode.
+		if constexpr (Version >= 12)
+		{
+			AttributeValidator.AddOptional(TEXT("seed"), ENNEAttributeDataType::Int32);
+		}
+		else
+		{
+			static_assert(Version >= 7, "Minimum supported version for operator Dropout is 7!");
+			AttributeValidator.AddOptional(TEXT("ratio"), ENNEAttributeDataType::Float);
+		}
 		bIsValid &= AttributeValidator.Validate(AttributeMap);
 
 		FInputValidator InputValidator;
@@ -100,8 +109,10 @@ namespace UE::NNERuntimeRDG::Private::Hlsl
 	bool RegisterDropoutOperator(FOperatorRegistryHlsl& Registry)
 	{
 		// Note: support of a particular version is partial with respect to tensor data types (only the most typical ones are usually supported).
-		Registry.OpAdd({{TEXT("Dropout"), TEXT("Onnx")}, 12}, CreateDropoutOperator, ValidateDropoutOperator);
-		Registry.OpAdd({{TEXT("Dropout"), TEXT("Onnx")}, 13}, CreateDropoutOperator, ValidateDropoutOperator);
+		Registry.OpAdd({{TEXT("Dropout"), TEXT("Onnx")}, 7}, CreateDropoutOperator, ValidateDropoutOperator<7>);
+		Registry.OpAdd({{TEXT("Dropout"), TEXT("Onnx")}, 10}, CreateDropoutOperator, ValidateDropoutOperator<10>);
+		Registry.OpAdd({{TEXT("Dropout"), TEXT("Onnx")}, 12}, CreateDropoutOperator, ValidateDropoutOperator<12>);
+		Registry.OpAdd({{TEXT("Dropout"), TEXT("Onnx")}, 13}, CreateDropoutOperator, ValidateDropoutOperator<13>);
 		return true;
 	}
 } // UE::NNERuntimeRDG::Private::Hlsl
