[NNE] Fixed issues with Resize and Upsample operators from NNERuntimeRDGDml runtime
Problem
- OIDN NNEQATest are failing with NNERuntimeRDGDml
- There are visual artifacts that shove that image is shifted by few pixels when running NNEDenoiser

Solution
- Update Resize and Upsample operators to use DML_RESAMPLE2_OPERATOR_DESC
- Added support for Resize nearest_mode attribute
- Changed how Resize operator computes output shape
- Make sure that OutputPadding is initialized in the Conv operator
- Added check() when to trigger when empty data is being used by operators that use FKernelArgs and FConvArgs
- Make sure that minimal buffer size is aligned to 4 bytes boundary

Tests
- NNEIncubatorEditor Win64 Debug/Development
- NNEIncubator game Win64
- Non-unity compile

#rnx

#virtualized

diff --git a/Engine/Plugins/Experimental/NNERuntimeRDG/Source/NNERuntimeRDG/Private/Dml/NNEDmlModel.cpp b/Engine/Plugins/Experimental/NNERuntimeRDG/Source/NNERuntimeRDG/Private/Dml/NNEDmlModel.cpp
--- a/Engine/Plugins/Experimental/NNERuntimeRDG/Source/NNERuntimeRDG/Private/Dml/NNEDmlModel.cpp
+++ b/Engine/Plugins/Experimental/NNERuntimeRDG/Source/NNERuntimeRDG/Private/Dml/NNEDmlModel.cpp
@@ -1083,9 +1083,10 @@ bool FModelInstance::InitCompiledOp()
 					const FTensorRDG&		Tensor = WeightTensorRDGs[WeightIdx];
 					TConstArrayView<uint8>	TensorData = Tensor.GetPreparedData<uint8>();
 
-					FBufferRHIRef WeightBuff;
+					FBufferRHIRef	WeightBuff;
+					const uint32	WeightBuffSize = Util::AlignBufferSize(TensorData.Num());
 					
-					WeightBuff = CreateRHIBuffer(RHICmdList, TensorData.Num(), WeightBuffUsage, WeightBuffAccess, *Tensor.GetName());
+					WeightBuff = CreateRHIBuffer(RHICmdList, WeightBuffSize, WeightBuffUsage, WeightBuffAccess, *Tensor.GetName());
 					check(WeightBuff);
 					WeightBuff->DisableLifetimeExtension();
 					
diff --git a/Engine/Plugins/Experimental/NNERuntimeRDG/Source/NNERuntimeRDG/Private/Dml/NNEDmlOperator.cpp b/Engine/Plugins/Experimental/NNERuntimeRDG/Source/NNERuntimeRDG/Private/Dml/NNEDmlOperator.cpp
--- a/Engine/Plugins/Experimental/NNERuntimeRDG/Source/NNERuntimeRDG/Private/Dml/NNEDmlOperator.cpp
+++ b/Engine/Plugins/Experimental/NNERuntimeRDG/Source/NNERuntimeRDG/Private/Dml/NNEDmlOperator.cpp
@@ -349,7 +349,7 @@ uint64 FTensorDescDml::CalculateBufferSize(uint64 ElemSizeInBytes)
 	}
 
 	// Round up to the nearest 4 bytes
-	MinSizeInBytes = (MinSizeInBytes + 3) & ~3ull;
+	MinSizeInBytes = Util::AlignBufferSize(MinSizeInBytes);
 
 	return MinSizeInBytes;
 }
diff --git a/Engine/Plugins/Experimental/NNERuntimeRDG/Source/NNERuntimeRDG/Private/Dml/NNEDmlOperator.h b/Engine/Plugins/Experimental/NNERuntimeRDG/Source/NNERuntimeRDG/Private/Dml/NNEDmlOperator.h
--- a/Engine/Plugins/Experimental/NNERuntimeRDG/Source/NNERuntimeRDG/Private/Dml/NNEDmlOperator.h
+++ b/Engine/Plugins/Experimental/NNERuntimeRDG/Source/NNERuntimeRDG/Private/Dml/NNEDmlOperator.h
@@ -81,6 +81,14 @@ class FTensorDescDml;
 
 namespace Util
 {
+
+// DML requires buffer size to be aligned at 4 byte boundary
+template<typename T>
+inline T AlignBufferSize(T BuffSizeInBytes)
+{
+	static_assert(std::is_integral<T>::value);
+	return (BuffSizeInBytes + 3) & ~3ull;
+}
 	
 template<typename T>
 using FSmallArray = TArray<T, TInlineAllocator<NNE::FTensorShape::MaxRank>>;
diff --git a/Engine/Plugins/Experimental/NNERuntimeRDG/Source/NNERuntimeRDG/Private/Dml/NNEDmlOperatorConvolution.cpp b/Engine/Plugins/Experimental/NNERuntimeRDG/Source/NNERuntimeRDG/Private/Dml/NNEDmlOperatorConvolution.cpp
--- a/Engine/Plugins/Experimental/NNERuntimeRDG/Source/NNERuntimeRDG/Private/Dml/NNEDmlOperatorConvolution.cpp
+++ b/Engine/Plugins/Experimental/NNERuntimeRDG/Source/NNERuntimeRDG/Private/Dml/NNEDmlOperatorConvolution.cpp
@@ -78,6 +78,7 @@ class FOperatorDmlConv : public FOperatorDml
 
 		TConstArrayView<uint32> GetOutPadding() const
 		{
+			check(!OutPadding.IsEmpty());
 			return MakeArrayView((const uint32*) OutPadding.GetData(), OutPadding.Num());
 		}
 
diff --git a/Engine/Plugins/Experimental/NNERuntimeRDG/Source/NNERuntimeRDG/Private/Dml/NNEDmlOperatorResample.cpp b/Engine/Plugins/Experimental/NNERuntimeRDG/Source/NNERuntimeRDG/Private/Dml/NNEDmlOperatorResample.cpp
--- a/Engine/Plugins/Experimental/NNERuntimeRDG/Source/NNERuntimeRDG/Private/Dml/NNEDmlOperatorResample.cpp
+++ b/Engine/Plugins/Experimental/NNERuntimeRDG/Source/NNERuntimeRDG/Private/Dml/NNEDmlOperatorResample.cpp
@@ -30,6 +30,15 @@ void RemoveValuesByIndex(TConstArrayView<uint32> Indices, TArray<TData, TAllocat
 	}
 }
 
+template<typename TData>
+inline TArray<TData, TInlineAllocator<GMaxTensorRank>> MakeFillArray(TData Value, int32 Count)
+{
+	TArray<TData, TInlineAllocator<GMaxTensorRank>>	Values;
+
+	Values.Init(Value, Count);
+	return Values;
+}
+
 // Upsample and Resize operator are implemented as a DML Resample operator
 template <bool IsResize>
 class FOperatorDmlResample : public FOperatorDml
@@ -44,9 +53,18 @@ class FOperatorDmlResample : public FOperatorDml
 		HalfPixel
 	};
 
-	DML_INTERPOLATION_MODE	Mode;
-	ECoordTransformMode		CoordTransformMode;
-	bool 					bUseSizesTensor;
+	enum ENearestNeighborRoundingMode : uint8
+	{
+		RoundPreferFloor,	// round halves down
+		RoundPreferCeil,	// round halves up
+		Floor,				// round always down
+		Ceil				// round always up
+	};
+
+	DML_INTERPOLATION_MODE			Mode;
+	ECoordTransformMode				CoordTransformMode;
+	ENearestNeighborRoundingMode	NearestMode;
+	bool 							bUseSizesTensor;
 
 	static DML_INTERPOLATION_MODE ModeFromString(FStringView StringVal)
 	{
@@ -91,10 +109,15 @@ public:
 
 			if(InputShapes.Num() > 1)
 			{
-				if (!CheckGenericTensor1D(OpName, InputTypes[1], InputShapes[1], 
-					{ 	ENNETensorDataType::Double, ENNETensorDataType::Float, ENNETensorDataType::Half
-					}
-					))
+				TSet<ENNETensorDataType> AllowedDataTypes = { ENNETensorDataType::Double, ENNETensorDataType::Float, ENNETensorDataType::Half };
+
+				// Add support for a scalar tensor which doesn't have a type
+				if (InputShapes[1].Rank() == 0 && InputTypes[1] == ENNETensorDataType::None)
+				{
+					AllowedDataTypes.Add(ENNETensorDataType::None);
+				}
+
+				if (!CheckGenericTensor1D(OpName, InputTypes[1], InputShapes[1], AllowedDataTypes))
 				{
 					return false;
 				}
@@ -194,11 +217,28 @@ public:
 		{
 			if (Mode == DML_INTERPOLATION_MODE_NEAREST_NEIGHBOR)
 			{
-				// DML only supports floor. If 'nearest_mode' wasn't specified, no need to warn the user.
-				FString NeareastMode = Attributes.GetValueOrDefault<FString>(TEXT("nearest_mode"), TEXT("floor"));
-				if (FCString::Stricmp(*NeareastMode, TEXT("floor")) != 0)
+				FString NearestModeStr = Attributes.GetValueOrDefault<FString>(TEXT("nearest_mode"), TEXT("round_prefer_floor"));
+				
+				if (NearestModeStr == TEXT("round_prefer_floor"))
+				{
+					NearestMode = ENearestNeighborRoundingMode::RoundPreferFloor;
+				}
+				else if (NearestModeStr == TEXT("round_prefer_ceil"))
+				{
+					NearestMode = ENearestNeighborRoundingMode::RoundPreferCeil;
+				}
+				else if (NearestModeStr == TEXT("floor"))
+				{
+					NearestMode = ENearestNeighborRoundingMode::Floor;
+				}
+				else if (NearestModeStr == TEXT("ceil"))
+				{
+					NearestMode = ENearestNeighborRoundingMode::Ceil;
+				}
+				else
 				{
-					UE_LOG(LogNNE, Display, TEXT("Unsupported neareast mode:%s, using floor instead"), *NeareastMode);
+					NearestMode = ENearestNeighborRoundingMode::Floor;
+					UE_LOG(LogNNE, Warning, TEXT("Unsupported neareast mode:%s, using floor instead"), *NearestModeStr);
 				}
 			}
 
@@ -435,8 +475,29 @@ public:
 			RemoveValuesByIndex(SqueezeInds, InputPixelOffsets, true);
 			RemoveValuesByIndex(SqueezeInds, OutputPixelOffsets, true);
 
-			DmlInputTensorDesc.SetShape(SqueezedInputShape);
-			DmlOutputTensorDesc.SetShape(SqueezedOutputShape);
+			if constexpr (IsResize)
+			{
+				const int32 SqueezedDimCount = SqueezedOutputShape.Num();
+				const int32 DmlDimCount = OutputTensor.GetShape().Rank();
+
+				if (DmlDimCount > SqueezedDimCount)
+				{
+					ScaleValues.Insert(MakeFillArray<float>(1.0f, DmlDimCount - SqueezedDimCount), 0);
+					InputPixelOffsets.Insert(MakeFillArray<float>(0.5f, DmlDimCount - SqueezedDimCount), 0);
+					OutputPixelOffsets.Insert(MakeFillArray<float>(-0.5f, DmlDimCount - SqueezedDimCount), 0);
+				}
+				else
+				{
+					DmlInputTensorDesc.SetShape(SqueezedInputShape);
+					DmlOutputTensorDesc.SetShape(SqueezedOutputShape);
+				}
+			}
+			else
+			{
+				DmlInputTensorDesc.SetShape(SqueezedInputShape);
+				DmlOutputTensorDesc.SetShape(SqueezedOutputShape);
+			}
+
 			Scales = ScaleValues;
 		}
 
@@ -452,17 +513,58 @@ public:
 			return false;
 		}
 
-		DML_RESAMPLE1_OPERATOR_DESC	OpDesc{};
+		DML_AXIS_DIRECTION RoundingDirection = DML_AXIS_DIRECTION_DECREASING;
+		if (Mode == DML_INTERPOLATION_MODE_NEAREST_NEIGHBOR)
+		{
+			float OffsetAdjustment = 0.5f;
+
+			switch (NearestMode)
+			{
+				case ENearestNeighborRoundingMode::RoundPreferFloor:
+					RoundingDirection = DML_AXIS_DIRECTION_INCREASING;
+					OffsetAdjustment = 0.5f;
+					break;
+
+				case ENearestNeighborRoundingMode::RoundPreferCeil:
+					RoundingDirection = DML_AXIS_DIRECTION_DECREASING;
+					OffsetAdjustment = -0.5f;
+					break;
+
+				case ENearestNeighborRoundingMode::Floor:
+					RoundingDirection = DML_AXIS_DIRECTION_DECREASING;
+					OffsetAdjustment = 0.0f;
+					break;
+
+				case ENearestNeighborRoundingMode::Ceil:
+					RoundingDirection = DML_AXIS_DIRECTION_INCREASING;
+					OffsetAdjustment = 0.0f;
+					break;
+
+				default:
+					UE_LOG(LogNNE, Warning, TEXT("Nearest neighbor rounding mode should have been initialized"));
+			}
+
+			if (OffsetAdjustment != 0.0f)
+			{
+				for (float& Offset : InputPixelOffsets)
+				{
+					Offset += OffsetAdjustment;
+				}
+			}
+		}
+
+		DML_RESAMPLE2_OPERATOR_DESC	OpDesc{};
 
 		OpDesc.InputTensor = DmlInputTensorDesc.GetDmlDesc();
 		OpDesc.OutputTensor = DmlOutputTensorDesc.GetDmlDesc();
 		OpDesc.InterpolationMode = Mode;
+		OpDesc.RoundingDirection = RoundingDirection;
 		OpDesc.DimensionCount = Scales.Num();
 		OpDesc.Scales = Scales.GetData();
 		OpDesc.InputPixelOffsets = InputPixelOffsets.GetData();
 		OpDesc.OutputPixelOffsets = OutputPixelOffsets.GetData();
 
-		return CreateOperator(Device, DML_OPERATOR_DESC { DML_OPERATOR_RESAMPLE1, &OpDesc });
+		return CreateOperator(Device, DML_OPERATOR_DESC{ DML_OPERATOR_RESAMPLE2, &OpDesc });
 	}
 };
 
diff --git a/Engine/Plugins/Experimental/NNERuntimeRDG/Source/NNERuntimeRDG/Private/Dml/NNEDmlOperatorUtils.h b/Engine/Plugins/Experimental/NNERuntimeRDG/Source/NNERuntimeRDG/Private/Dml/NNEDmlOperatorUtils.h
--- a/Engine/Plugins/Experimental/NNERuntimeRDG/Source/NNERuntimeRDG/Private/Dml/NNEDmlOperatorUtils.h
+++ b/Engine/Plugins/Experimental/NNERuntimeRDG/Source/NNERuntimeRDG/Private/Dml/NNEDmlOperatorUtils.h
@@ -232,26 +232,31 @@ public:
 
 	TConstArrayView<uint32> GetStrides() const
 	{
+		check(!Strides.IsEmpty());
 		return Strides;
 	}
 
 	TConstArrayView<uint32> GetDilations() const
 	{
+		check(!Dilations.IsEmpty());
 		return Dilations;
 	}
 
 	TConstArrayView<uint32> GetStartPadding()
 	{
+		check(!StartPadding.IsEmpty());
 		return StartPadding;
 	}
 
 	TConstArrayView<uint32> GetEndPadding()
 	{
+		check(!EndPadding.IsEmpty());
 		return EndPadding;
 	}
 
 	TConstArrayView<uint32> GetOutputShape() const
 	{
+		check(!OutputShape.IsEmpty());
 		return OutputShape;
 	}
 
@@ -354,6 +359,10 @@ protected:
 
 				bHasOutputShape = AttrOutShape != nullptr;
 			}
+			else
+			{
+				OutPadding.Init(0, NumDimensions);
+			}
 
 			if (!Paddings.Init(Attributes, InputShapeRank))
 			{
