[NNE] NNERuntimeRDG Buffer debug categorized under NNE.

Problem:
NNE Buffer were hard to track as the debug name was the tensor name, witch can be anything also CreateBuffer use a pointer to the debug name meaning the name should be kept in memory at all time.

Solution:
Use compile time constant to categorize NNE buffer.

Test: Win64 Editor: smoke tests + render resource viewer to track NNE buffer allocation.
Test: Win64 Standalone: all tests

#rnx

diff --git a/Engine/Plugins/Experimental/NNERuntimeRDG/Source/NNERuntimeRDG/Private/NNERuntimeRDGModel.cpp b/Engine/Plugins/Experimental/NNERuntimeRDG/Source/NNERuntimeRDG/Private/NNERuntimeRDGModel.cpp
--- a/Engine/Plugins/Experimental/NNERuntimeRDG/Source/NNERuntimeRDG/Private/NNERuntimeRDGModel.cpp
+++ b/Engine/Plugins/Experimental/NNERuntimeRDG/Source/NNERuntimeRDG/Private/NNERuntimeRDGModel.cpp
@@ -232,7 +232,7 @@ int32 FModelInstanceRDG::EnqueueRDG(FRDGBuilder& RDGBuilder, TConstArrayView<NNE
 		if (!TensorRDG.HasPreparedData())
 		{
 			const FRDGBufferDesc BufferDesc = CreateRDGBufferDescForTensorRDG(TensorRDG);
-			const FRDGBufferRef TensorBuffer = RDGBuilder.CreateBuffer(BufferDesc, *TensorRDG.GetName(), ERDGBufferFlags::None);
+			const FRDGBufferRef TensorBuffer = RDGBuilder.CreateBuffer(BufferDesc, TEXT("NNE.Tensor.Intermediate"), ERDGBufferFlags::None);
 			check(!bBuffersUploadedAndRegisteredToRDGGraph || TensorRDG.GetBuffer() == nullptr);
 			TensorRDG.SetBuffer(TensorBuffer);
 		}
diff --git a/Engine/Plugins/Experimental/NNERuntimeRDG/Source/NNERuntimeRDG/Private/NNERuntimeRDGModelHlsl.cpp b/Engine/Plugins/Experimental/NNERuntimeRDG/Source/NNERuntimeRDG/Private/NNERuntimeRDGModelHlsl.cpp
--- a/Engine/Plugins/Experimental/NNERuntimeRDG/Source/NNERuntimeRDG/Private/NNERuntimeRDGModelHlsl.cpp
+++ b/Engine/Plugins/Experimental/NNERuntimeRDG/Source/NNERuntimeRDG/Private/NNERuntimeRDGModelHlsl.cpp
@@ -74,7 +74,7 @@ bool FModelInstance::PrepareModelRDG(FRDGBuilder& RDGBuilder)
 			if (!PooledBuffer.IsValid())
 			{
 				FRDGBufferDesc BufferDesc = FRDGBufferDesc::CreateBufferDesc(Tensor.GetElementByteSize(), Tensor.GetVolume());
-				const FRDGBufferRef TransientRDGBuffer = RDGBuilder.CreateBuffer(BufferDesc, *Tensor.GetName(), ERDGBufferFlags::None);
+				const FRDGBufferRef TransientRDGBuffer = RDGBuilder.CreateBuffer(BufferDesc, TEXT("NNE.Tensor.Intermediate.Constant"), ERDGBufferFlags::None);
 				const uint8* TensorData = Tensor.GetPreparedData<uint8>().GetData();
 				PooledBuffer = RDGBuilder.ConvertToExternalBuffer(TransientRDGBuffer);
 
@@ -327,7 +327,7 @@ void EnqueueTensorUpload(TArray<TRefCountPtr<FRDGPooledBuffer>>& OutExternalRDGR
 					{
 
 						FRDGBufferDesc BufferDesc = FRDGBufferDesc::CreateBufferDesc(Tensor.GetElementByteSize(), Tensor.GetVolume());
-						const FRDGBufferRef TransientRDGBuffer = RDGBuilder.CreateBuffer(BufferDesc, *Tensor.GetName(), ERDGBufferFlags::None);
+						const FRDGBufferRef TransientRDGBuffer = RDGBuilder.CreateBuffer(BufferDesc, TEXT("NNE.Tensor.Weights"), ERDGBufferFlags::None);
 						const uint8* TensorData = Tensor.GetPreparedData<uint8>().GetData();
 
 						OutExternalRDGResources[i] = RDGBuilder.ConvertToExternalBuffer(TransientRDGBuffer);
