[NNE] NNERuntimeRDGHlsl Support Splice opset 10

Problem:
Slice was only supported up to opset 9.

Solution:
Refactor the op to support up to latest opset with the limitations to use constant tensors for stard/end/axes and step of 1.

Test: Win64 Editor: slice tests + smoke tests
Test: Win64 Standalone: all tests

#rnx

diff --git a/Engine/Plugins/Experimental/NNERuntimeRDG/Source/NNERuntimeRDG/Private/Helper/NNERuntimeRDGOperatorHelper.cpp b/Engine/Plugins/Experimental/NNERuntimeRDG/Source/NNERuntimeRDG/Private/Helper/NNERuntimeRDGOperatorHelper.cpp
new file mode 100644
--- /dev/null
+++ b/Engine/Plugins/Experimental/NNERuntimeRDG/Source/NNERuntimeRDG/Private/Helper/NNERuntimeRDGOperatorHelper.cpp
@@ -0,0 +1,38 @@
+// Copyright Epic Games, Inc. All Rights Reserved.
+
+#include "NNERuntimeRDGOperatorHelper.h"
+
+namespace UE::NNERuntimeRDG::Private::OperatorHelper
+{
+
+	bool GetInt32ArrayFromConstTensor(TArray<int32, TInlineAllocator<NNE::FTensorShape::MaxRank>>& Attr, const NNE::Internal::FTensorRef Tensor)
+	{
+		check(Tensor != nullptr);
+		Attr.Reset();
+
+		if (!Tensor->HasPreparedData())
+		{
+			return false;
+		}
+
+		if (Tensor->GetDataType() == ENNETensorDataType::Int64)
+		{
+			for (int64 Value64 : Tensor->GetPreparedData<int64>())
+			{
+				int64 ValueClamped = FMath::Clamp<int64>(Value64, MIN_int32, MAX_int32);
+				Attr.Add((int32)ValueClamped);
+			}
+
+			return true;
+		}
+		else if (Tensor->GetDataType() == ENNETensorDataType::Int32)
+		{
+			Attr.Append(Tensor->GetPreparedData<int32>());
+			return true;
+		}
+
+		return false;
+	}
+
+} // UE::NNERuntimeRDG::Private::OperatorHelper
+
diff --git a/Engine/Plugins/Experimental/NNERuntimeRDG/Source/NNERuntimeRDG/Private/Helper/NNERuntimeRDGOperatorHelper.h b/Engine/Plugins/Experimental/NNERuntimeRDG/Source/NNERuntimeRDG/Private/Helper/NNERuntimeRDGOperatorHelper.h
new file mode 100644
--- /dev/null
+++ b/Engine/Plugins/Experimental/NNERuntimeRDG/Source/NNERuntimeRDG/Private/Helper/NNERuntimeRDGOperatorHelper.h
@@ -0,0 +1,14 @@
+// Copyright Epic Games, Inc. All Rights Reserved.
+
+#pragma once
+
+#include "CoreMinimal.h"
+#include "NNETensor.h"
+#include "NNETypes.h"
+
+namespace UE::NNERuntimeRDG::Private::OperatorHelper
+{
+	bool GetInt32ArrayFromConstTensor(TArray<int32, TInlineAllocator<NNE::FTensorShape::MaxRank>>& Attr, const NNE::Internal::FTensorRef Tensor);
+
+} // UE::NNERuntimeRDG::Private::OperatorHelper
+
diff --git a/Engine/Plugins/Experimental/NNERuntimeRDG/Source/NNERuntimeRDG/Private/Hlsl/NNERuntimeRDGSlice.cpp b/Engine/Plugins/Experimental/NNERuntimeRDG/Source/NNERuntimeRDG/Private/Hlsl/NNERuntimeRDGSlice.cpp
--- a/Engine/Plugins/Experimental/NNERuntimeRDG/Source/NNERuntimeRDG/Private/Hlsl/NNERuntimeRDGSlice.cpp
+++ b/Engine/Plugins/Experimental/NNERuntimeRDG/Source/NNERuntimeRDG/Private/Hlsl/NNERuntimeRDGSlice.cpp
@@ -1,6 +1,7 @@
 // Copyright Epic Games, Inc. All Rights Reserved.
 
 #include "NNERuntimeRDGSlice.h"
+#include "Helper/NNERuntimeRDGOperatorHelper.h"
 #include "NNEHlslShadersSliceCS.h"
 #include "NNERuntimeRDGHelperSlice.h"
 #include "NNERuntimeRDGHlslHelper.h"
@@ -15,6 +16,7 @@ namespace UE::NNERuntimeRDG::Private::Hlsl
 	/**
 	 * Slice operator implementation
 	 */
+	template<int32 OpVersion>
 	class FSlice : public FOperatorHlsl
 	{
 	public:
@@ -24,14 +26,64 @@ namespace UE::NNERuntimeRDG::Private::Hlsl
 
 	private:
 
-		TArray<int32> AxesAttr;
-		TArray<int32> EndsAttr;
-		TArray<int32> StartsAttr;
+		mutable TArray<int32, TInlineAllocator<NNE::FTensorShape::MaxRank>> AxesAttr;
+		mutable TArray<int32, TInlineAllocator<NNE::FTensorShape::MaxRank>> EndsAttr;
+		mutable TArray<int32, TInlineAllocator<NNE::FTensorShape::MaxRank>> StartsAttr;
 
-		inline void GetStartAndEndFromInputShape(TConstArrayView<uint32> InputShapeData, 
-			TArray<int32, TInlineAllocator<NNE::FTensorShape::MaxRank>>& Start,
-			TArray<int32, TInlineAllocator<NNE::FTensorShape::MaxRank>>& End) const
+		mutable TArray<int32, TInlineAllocator<NNE::FTensorShape::MaxRank>> Start;
+		mutable TArray<int32, TInlineAllocator<NNE::FTensorShape::MaxRank>> End;
+
+		bool TryGetAttributesFromConstantTensors(TConstArrayView<NNE::Internal::FTensorRef> InputTensors) const
+		{
+			check(InputTensors.Num() >= 3 && InputTensors.Num() <= 5);
+
+			if (!OperatorHelper::GetInt32ArrayFromConstTensor(StartsAttr, InputTensors[1]))
+			{
+				UE_LOG(LogNNE, Warning, TEXT("Error: Slice op 'Starts' input at index 1 is only supported as a constant integer tensor but it is not."));
+				return false;
+			}
+			if (!OperatorHelper::GetInt32ArrayFromConstTensor(EndsAttr, InputTensors[2]))
+			{
+				UE_LOG(LogNNE, Warning, TEXT("Error: Slice op 'Ends' input at index 2 is only supported as a constant integer tensor but it is not."));
+				return false;
+			}
+			if (InputTensors.Num() >= 4)
+			{
+				if (!OperatorHelper::GetInt32ArrayFromConstTensor(AxesAttr, InputTensors[3]))
+				{
+					UE_LOG(LogNNE, Warning, TEXT("Error: Slice op 'Axes' input at index 3 is only supported as a constant integer tensor but it is not."));
+					return false;
+				}
+			}
+			else
+			{
+				for (int32 i = 0; i < StartsAttr.Num(); ++i)
+				{
+					AxesAttr.Add(i);
+				}
+			}
+			if (InputTensors.Num() == 5)
+			{
+				TArray<int32, TInlineAllocator<NNE::FTensorShape::MaxRank>> StepsAttr;
+				OperatorHelper::GetInt32ArrayFromConstTensor(StepsAttr, InputTensors[4]);
+				for (int32 Value : StepsAttr)
+				{
+					if (Value != 1)
+					{
+						UE_LOG(LogNNE, Warning, TEXT("Error: Slice op 'Steps' optional input at index 4 is only supported as a constant tensor will all values as 1."));
+						return false;
+					}
+				}
+			}
+
+			return true;
+		}
+
+		void ComputeStartAndEndFromInputShape(TConstArrayView<uint32> InputShapeData) const
 		{
+			check(AxesAttr.Num() == EndsAttr.Num());
+			check(AxesAttr.Num() == StartsAttr.Num());
+
 			const int32 InputRank = InputShapeData.Num();
 			TArray<int32, TInlineAllocator<NNE::FTensorShape::MaxRank>> Axes(AxesAttr);
 			TArray<int32, TInlineAllocator<NNE::FTensorShape::MaxRank>> Ends(EndsAttr);
@@ -77,15 +129,25 @@ namespace UE::NNERuntimeRDG::Private::Hlsl
 
 		virtual int PrepareOutputs(TConstArrayView<NNE::Internal::FTensorRef> InputTensors, TArrayView<NNE::Internal::FTensorRef> OutputTensors) const override
 		{
-			check(InputTensors.Num() == 1);
 			check(OutputTensors.Num() == 1);
 
+			if constexpr (OpVersion == 1)
+			{
+				check(InputTensors.Num() == 1);
+			}
+			else
+			{
+				check(InputTensors.Num() >= 3 && InputTensors.Num() <= 5);
+				if (!TryGetAttributesFromConstantTensors(InputTensors))
+				{
+					return -1;
+				}
+			}
+
 			TConstArrayView<uint32> InputShapeData(InputTensors[0]->GetShape().GetData());
 			const int32 InputRank = InputShapeData.Num();
-			TArray<int32, TInlineAllocator<NNE::FTensorShape::MaxRank>> Start(AxesAttr);
-			TArray<int32, TInlineAllocator<NNE::FTensorShape::MaxRank>> End(EndsAttr);
 
-			GetStartAndEndFromInputShape(InputShapeData, Start, End);
+			ComputeStartAndEndFromInputShape(InputShapeData);
 
 			TArray<uint32> OutputShapeData;
 
@@ -101,30 +163,44 @@ namespace UE::NNERuntimeRDG::Private::Hlsl
 			
 			Internal::CPUHelper::Slice::Apply(*InputTensors[0], *OutputTensors[0], Start);
 
+			if (OutputTensors[0]->GetDataType() != ENNETensorDataType::Float && !OutputTensors[0]->HasPreparedData())
+			{
+				UE_LOG(LogNNE, Warning, TEXT("Error: Slice op output tensor could not be made constant nor it was of float type. Only floats are supported at the moment on the HLSL compute path."));
+				return -1;
+			}
+
 			return 0;
 		};
 
 		virtual bool Initialize(TConstArrayView<NNE::FTensorDesc> InputTensorDescs, TConstArrayView<NNE::FTensorDesc> OutputTensorDescs, const NNE::FAttributeMap& Attributes) override
 		{
-			check(InputTensorDescs.Num() == 1);
 			check(OutputTensorDescs.Num() == 1);
 
-			EndsAttr = Attributes.GetValue<TArray<int32>>(TEXT("ends"));
-			StartsAttr = Attributes.GetValue<TArray<int32>>(TEXT("starts"));
-			
-			TArray<int32> AxesDefault;
-			
-			for (int i = 0; i < StartsAttr.Num(); ++i)
+			if constexpr (OpVersion == 1)
 			{
-				AxesDefault.Add(i);
-			}
-			
-			AxesAttr = Attributes.GetValueOrDefault<TArray<int32>>(TEXT("axes"), AxesDefault);
+				check(InputTensorDescs.Num() == 1);
+
+				EndsAttr = Attributes.GetValue<TArray<int32>>(TEXT("ends"));
+				StartsAttr = Attributes.GetValue<TArray<int32>>(TEXT("starts"));
 
-			if (EndsAttr.Num() != StartsAttr.Num() || AxesAttr.Num() != StartsAttr.Num())
+				TArray<int32> AxesDefault;
+
+				for (int32 i = 0; i < StartsAttr.Num(); ++i)
+				{
+					AxesDefault.Add(i);
+				}
+
+				AxesAttr = Attributes.GetValueOrDefault<TArray<int32>>(TEXT("axes"), AxesDefault);
+
+				if (EndsAttr.Num() != StartsAttr.Num() || AxesAttr.Num() != StartsAttr.Num())
+				{
+					UE_LOG(LogNNE, Warning, TEXT("Slice: Starts, Ends and Axes must be of the same size."));
+					return false;
+				}
+			}
+			else
 			{
-				UE_LOG(LogNNE, Warning, TEXT("Unsqueeze: Starts, Ends and Axes must be of the same size."));
-				return false;
+				check(InputTensorDescs.Num() >= 3 && InputTensorDescs.Num() <= 5);
 			}
 			
 			return true;
@@ -134,7 +210,7 @@ namespace UE::NNERuntimeRDG::Private::Hlsl
 		{
 			using namespace UE::NNEHlslShaders::Internal;
 
-			check(InputTensors.Num() == 1);
+			check(InputTensors.Num() >= 1);
 			check(OutputTensors.Num() == 1);
 			check(InputTensors[0] != nullptr);
 			check(OutputTensors[0] != nullptr);
@@ -145,12 +221,6 @@ namespace UE::NNERuntimeRDG::Private::Hlsl
 			const FRDGBufferUAVRef OutputUAV = GraphBuilder.CreateUAV(FRDGBufferUAVDesc(Output.GetBuffer(), PF_R32_FLOAT));
 			const FIntVector ThreadGroupCount = ComputeElementWiseThreadGroups(Output.GetVolume(), FSliceConstants::NUM_GROUP_THREADS);
 
-			TConstArrayView<uint32> InputShapeData(Input.GetShape().GetData());
-			TArray<int32, TInlineAllocator<NNE::FTensorShape::MaxRank>> Start(AxesAttr);
-			TArray<int32, TInlineAllocator<NNE::FTensorShape::MaxRank>> End(EndsAttr);
-
-			GetStartAndEndFromInputShape(InputShapeData, Start, End);
-
 			// Set parameters
 			FSliceCS::FParameters* Params = GraphBuilder.AllocParameters<FSliceCS::FParameters>();
 			Params->Input = InputSRV;
@@ -191,7 +261,7 @@ namespace UE::NNERuntimeRDG::Private::Hlsl
 		}
 	};
 
-	bool ValidateSliceOperator(const NNE::FAttributeMap& AttributeMap, TConstArrayView<ENNETensorDataType> InputTypes, TConstArrayView<NNE::FSymbolicTensorShape> InputSlices)
+	bool ValidateSliceOperatorOpset1(const NNE::FAttributeMap& AttributeMap, TConstArrayView<ENNETensorDataType> InputTypes, TConstArrayView<NNE::FSymbolicTensorShape> InputSlices)
 	{
 		bool bIsValid = true;
 
@@ -212,15 +282,46 @@ namespace UE::NNERuntimeRDG::Private::Hlsl
 		return bIsValid;
 	}
 
+	bool ValidateSliceOperatorOpset10To13(const NNE::FAttributeMap& AttributeMap, TConstArrayView<ENNETensorDataType> InputTypes, TConstArrayView<NNE::FSymbolicTensorShape> InputSlices)
+	{
+		bool bIsValid = true;
+
+		FAttributeValidator AttributeValidator;
+		bIsValid &= AttributeValidator.Validate(AttributeMap);
+
+		FInputValidator InputValidator;
+		InputValidator.SetTemplateCount(2);
+
+		InputValidator.AddSupportedType(ENNETensorDataType::Float, 0);
+		InputValidator.AddSupportedType(ENNETensorDataType::Int32, 0);
+		InputValidator.AddSupportedType(ENNETensorDataType::Int64, 0);
+		InputValidator.AddRequired(0);//Data
+
+		InputValidator.AddSupportedType(ENNETensorDataType::Int32, 1);
+		InputValidator.AddSupportedType(ENNETensorDataType::Int64, 1);
+		InputValidator.AddRequired(1);//Starts
+		InputValidator.AddRequired(1);//Ends
+		InputValidator.AddOptional(1);//Axes
+		InputValidator.AddOptional(1);//Steps
+
+		bIsValid &= InputValidator.Validate(InputTypes);
+
+		return bIsValid;
+	}
+
+	template<int32 OpsetVersion>
 	FOperatorHlsl* CreateSliceOperator()
 	{
-		return new FSlice();
+		return new FSlice<OpsetVersion>();
 	}
 
 	bool RegisterSliceOperator(FOperatorRegistryHlsl& Registry)
 	{
 		// Note: support of a particular version is partial with respect to tensor data types (only the most typical ones are usually supported).
-		Registry.OpAdd({{TEXT("Slice"), TEXT("Onnx")}, 1}, CreateSliceOperator, ValidateSliceOperator);
+		Registry.OpAdd({ {TEXT("Slice"), TEXT("Onnx")}, 1  }, CreateSliceOperator<1>,  ValidateSliceOperatorOpset1);
+		Registry.OpAdd({ {TEXT("Slice"), TEXT("Onnx")}, 10 }, CreateSliceOperator<10>, ValidateSliceOperatorOpset10To13);
+		Registry.OpAdd({ {TEXT("Slice"), TEXT("Onnx")}, 11 }, CreateSliceOperator<11>, ValidateSliceOperatorOpset10To13);
+		Registry.OpAdd({ {TEXT("Slice"), TEXT("Onnx")}, 13 }, CreateSliceOperator<13>, ValidateSliceOperatorOpset10To13);
 		return true;
 	}
 } // UE::NNERuntimeRDG::Private::Hlsl
