[NNE] NNERuntimeRDGHlsl Upsample operator 'linear' mode support.

Tests:
* P4V Non-Unity build
* UnrealEditor, Win64: Upsample tests, smoke tests
* UnrealGame, Win64: all tests

#rnx

#virtualized

diff --git a/Engine/Plugins/Experimental/NNERuntimeRDG/Source/NNEHlslShaders/Internal/NNEHlslShadersUpsampleCS.h b/Engine/Plugins/Experimental/NNERuntimeRDG/Source/NNEHlslShaders/Internal/NNEHlslShadersUpsampleCS.h
--- a/Engine/Plugins/Experimental/NNERuntimeRDG/Source/NNEHlslShaders/Internal/NNEHlslShadersUpsampleCS.h
+++ b/Engine/Plugins/Experimental/NNERuntimeRDG/Source/NNEHlslShaders/Internal/NNEHlslShadersUpsampleCS.h
@@ -8,6 +8,14 @@
 
 namespace UE::NNEHlslShaders::Internal
 {
+	enum class EUpsampleMode : uint8
+	{
+		Nearest = 0,
+		Bilinear,
+		Trilinear,
+		MAX
+	};
+
 	class FUpsampleConstants
 	{
 	public:
@@ -22,7 +30,8 @@ namespace UE::NNEHlslShaders::Internal
 		SHADER_USE_PARAMETER_STRUCT(FUpsampleCS, FHlslShaderBase)
 
 		class FUpsampleNumDimensions : SHADER_PERMUTATION_RANGE_INT("NUM_DIMENSIONS", 1, FUpsampleConstants::MAX_NUM_DIMENSIONS);
-		using FPermutationDomain = TShaderPermutationDomain<FUpsampleNumDimensions>;
+		class FUpsampleMode : SHADER_PERMUTATION_ENUM_CLASS("MODE", EUpsampleMode);
+		using FPermutationDomain = TShaderPermutationDomain<FUpsampleNumDimensions, FUpsampleMode>;
 
 	public:
 
@@ -34,6 +43,8 @@ namespace UE::NNEHlslShaders::Internal
 			SHADER_PARAMETER(uint32, ThreadCountX)
 		END_SHADER_PARAMETER_STRUCT()
 
+		static bool ShouldCompilePermutation(const FGlobalShaderPermutationParameters& Parameters);
+
 		static void ModifyCompilationEnvironment(const FGlobalShaderPermutationParameters& InParameters, FShaderCompilerEnvironment& OutEnvironment);
 	};
 } // UE::NNEHlslShaders::Internal
diff --git a/Engine/Plugins/Experimental/NNERuntimeRDG/Source/NNEHlslShaders/Private/NNEHlslShadersUpsampleCS.cpp b/Engine/Plugins/Experimental/NNERuntimeRDG/Source/NNEHlslShaders/Private/NNEHlslShadersUpsampleCS.cpp
--- a/Engine/Plugins/Experimental/NNERuntimeRDG/Source/NNEHlslShaders/Private/NNEHlslShadersUpsampleCS.cpp
+++ b/Engine/Plugins/Experimental/NNERuntimeRDG/Source/NNEHlslShaders/Private/NNEHlslShadersUpsampleCS.cpp
@@ -5,6 +5,23 @@
 
 namespace UE::NNEHlslShaders::Internal
 {
+	bool FUpsampleCS::ShouldCompilePermutation(const FGlobalShaderPermutationParameters& Parameters)
+	{
+		if (!FHlslShaderBase::ShouldCompilePermutation(Parameters))
+		{
+			return false;
+		}
+
+		FPermutationDomain PermutationVector(Parameters.PermutationId);
+
+		EUpsampleMode Mode = PermutationVector.Get<FUpsampleCS::FUpsampleMode>();
+		const int32 NumDimensions = PermutationVector.Get<FUpsampleCS::FUpsampleNumDimensions>();
+
+		return Mode == EUpsampleMode::Nearest ||
+			(Mode == EUpsampleMode::Bilinear && NumDimensions > 1) ||
+			(Mode == EUpsampleMode::Trilinear && NumDimensions > 2);
+	}
+
 	void FUpsampleCS::ModifyCompilationEnvironment(const FGlobalShaderPermutationParameters& InParameters, FShaderCompilerEnvironment& OutEnvironment)
 	{
 		FGlobalShader::ModifyCompilationEnvironment(InParameters, OutEnvironment);
diff --git a/Engine/Plugins/Experimental/NNERuntimeRDG/Source/NNEHlslShaders/Shaders/NNEHlslShadersUpsample.usf b/Engine/Plugins/Experimental/NNERuntimeRDG/Source/NNEHlslShaders/Shaders/NNEHlslShadersUpsample.usf
--- a/Engine/Plugins/Experimental/NNERuntimeRDG/Source/NNEHlslShaders/Shaders/NNEHlslShadersUpsample.usf
+++ b/Engine/Plugins/Experimental/NNERuntimeRDG/Source/NNEHlslShaders/Shaders/NNEHlslShadersUpsample.usf
@@ -12,37 +12,198 @@ uint ThreadCountX;
 #define INPUT_STRIDE 0
 #define OUTPUT_STRIDE 1
 #define INPUT_SIZE 2
-#define OUTPUT_SIZE 3
+#define SCALES_DIV 3
 
-uint GetNearestIndexFromUV(float UV, uint ArraySize)
+#if MODE == 0 // Nearest
+
+[numthreads(THREADGROUP_SIZE_X, 1, 1)]
+void Upsample(in const uint3 DispatchThreadID : SV_DispatchThreadID)
 {
-	uint Index = round(UV * ArraySize - 0.5f);
-	return min(Index, ArraySize - 1);
+	const uint OutputIndex = DispatchThreadID.y * ThreadCountX + DispatchThreadID.x;
+	
+	if (OutputIndex < Num)
+	{
+		uint InputIndex = 0;
+
+		uint CurrIndex = OutputIndex;
+		for (uint dim = 0; dim < NUM_DIMENSIONS; ++dim)
+		{
+			uint OutputCoord, OutputOffset;
+			DivMod(CurrIndex, TensorInfo[dim][OUTPUT_STRIDE], OutputCoord, OutputOffset);
+
+			uint InputCoord, InputOffset;
+			DivMod(OutputCoord, TensorInfo[dim][SCALES_DIV], InputCoord, InputOffset);
+			
+			InputIndex += TensorInfo[dim][INPUT_STRIDE] * InputCoord;
+			CurrIndex = OutputOffset;
+		}
+
+		Output[OutputIndex] = Input[InputIndex];
+	}
 }
 
+#elif MODE == 1 // Bilinear
+
 [numthreads(THREADGROUP_SIZE_X, 1, 1)]
 void Upsample(in const uint3 DispatchThreadID : SV_DispatchThreadID)
 {
-	const uint Index = DispatchThreadID.y * ThreadCountX + DispatchThreadID.x;
+	const uint OutputIndex = DispatchThreadID.y * ThreadCountX + DispatchThreadID.x;
 	
-	if (Index < Num)
+	if (OutputIndex < Num)
 	{
 		uint InputIndex = 0;
 
-		uint Offset = Index;
-		for (uint dim = 0; dim < NUM_DIMENSIONS; ++dim)
+		uint CurrIndex = OutputIndex;
+		for (uint dim = 0; dim < NUM_DIMENSIONS - 2; ++dim)
 		{
-			uint OutDimIdx, R;
-			DivMod(Offset, TensorInfo[dim][OUTPUT_STRIDE], OutDimIdx, R);
+			uint OutputCoord, OutputOffset;
+			DivMod(CurrIndex, TensorInfo[dim][OUTPUT_STRIDE], OutputCoord, OutputOffset);
+
+			uint InputCoord, InputOffset;
+			DivMod(OutputCoord, TensorInfo[dim][SCALES_DIV], InputCoord, InputOffset);
 			
-			float UV = (0.5f + OutDimIdx) / TensorInfo[dim][OUTPUT_SIZE];
-			uint InputDimIdx = GetNearestIndexFromUV(UV, TensorInfo[dim][INPUT_SIZE]);
+			InputIndex += TensorInfo[dim][INPUT_STRIDE] * InputCoord;
+			CurrIndex = OutputOffset;
+		}
+
+		const uint dim0 = NUM_DIMENSIONS - 2;
+		const uint dim1 = NUM_DIMENSIONS - 1;
+
+		uint OutputCoord0, OutputCoord1;
+		DivMod(CurrIndex, TensorInfo[dim0][OUTPUT_STRIDE], OutputCoord0, OutputCoord1);
+
+		uint InputCoord0, InputCoord1, InputOffset0, InputOffset1;
+		DivMod(OutputCoord0, TensorInfo[dim0][SCALES_DIV], InputCoord0, InputOffset0);
+		DivMod(OutputCoord1, TensorInfo[dim1][SCALES_DIV], InputCoord1, InputOffset1);
+
+		InputIndex += TensorInfo[dim0][INPUT_STRIDE] * InputCoord0 + InputCoord1;
+
+		const float x00 = Input[InputIndex];
+		float x10, x01, x11;
+
+		bool DimEnd0 = false;
+		if (InputCoord0 == TensorInfo[dim0][INPUT_SIZE] - 1)
+		{
+			x01 = x00;
+			DimEnd0 = true;
+		}
+		else
+		{
+			x01 = Input[InputIndex + TensorInfo[dim0][INPUT_STRIDE]];
+		}
+
+		if (InputCoord1 == TensorInfo[dim1][INPUT_SIZE] - 1)
+		{
+			x10 = x00;
+			x11 = x01;
+		}
+		else
+		{
+			x10 = Input[InputIndex + 1];
+			x11 = DimEnd0 ? x10 : Input[InputIndex + TensorInfo[dim0][INPUT_STRIDE] + 1];
+		}
+
+		float y0 = x00 + (float)InputOffset0 * (x01 - x00) / (float)TensorInfo[dim0][SCALES_DIV];
+		float y1 = x10 + (float)InputOffset0 * (x11 - x10) / (float)TensorInfo[dim0][SCALES_DIV];
+
+		Output[OutputIndex] = y0 + (float)InputOffset1 * (y1 - y0) / (float)TensorInfo[dim1][SCALES_DIV];
+	}
+}
+
+#else // Trilinear
+
+[numthreads(THREADGROUP_SIZE_X, 1, 1)]
+void Upsample(in const uint3 DispatchThreadID : SV_DispatchThreadID)
+{
+	const uint OutputIndex = DispatchThreadID.y * ThreadCountX + DispatchThreadID.x;
+	
+	if (OutputIndex < Num)
+	{
+		uint InputIndex = 0;
+
+		uint CurrIndex = OutputIndex;
+		for (uint dim = 0; dim < NUM_DIMENSIONS - 3; ++dim)
+		{
+			uint OutputCoord, OutputOffset;
+			DivMod(CurrIndex, TensorInfo[dim][OUTPUT_STRIDE], OutputCoord, OutputOffset);
+
+			uint InputCoord, InputOffset;
+			DivMod(OutputCoord, TensorInfo[dim][SCALES_DIV], InputCoord, InputOffset);
 			
-			InputIndex += TensorInfo[dim][INPUT_STRIDE] * InputDimIdx;
-			Offset = R;
+			InputIndex += TensorInfo[dim][INPUT_STRIDE] * InputCoord;
+			CurrIndex = OutputOffset;
 		}
 
-		float X = Input[InputIndex];
-		Output[Index] = X;
+		const uint dim0 = NUM_DIMENSIONS - 3;
+		const uint dim1 = NUM_DIMENSIONS - 2;
+		const uint dim2 = NUM_DIMENSIONS - 1;
+
+		uint OutputCoord0, Mod, OutputCoord1, OutputCoord2;
+		DivMod(CurrIndex, TensorInfo[dim0][OUTPUT_STRIDE], OutputCoord0, Mod);
+		DivMod(Mod, TensorInfo[dim0][OUTPUT_STRIDE], OutputCoord1, OutputCoord2);
+
+		uint InputCoord0, InputCoord1, InputCoord2, InputOffset0, InputOffset1, InputOffset2;
+		DivMod(OutputCoord0, TensorInfo[dim0][SCALES_DIV], InputCoord0, InputOffset0);
+		DivMod(OutputCoord1, TensorInfo[dim1][SCALES_DIV], InputCoord1, InputOffset1);
+		DivMod(OutputCoord2, TensorInfo[dim2][SCALES_DIV], InputCoord2, InputOffset2);
+
+		InputIndex += TensorInfo[dim0][INPUT_STRIDE] * InputCoord0 +
+			TensorInfo[dim1][INPUT_STRIDE] * InputCoord1 +
+			InputCoord2;
+
+		const float x000 = Input[InputIndex];
+		float x010, x001, x011;
+		float x100, x110, x101, x111;
+
+		bool DimEnd0 = false;
+		if (InputCoord0 == TensorInfo[dim0][INPUT_SIZE] - 1)
+		{
+			x001 = x000;
+			DimEnd0 = true;
+		}
+		else
+		{
+			x001 = Input[InputIndex + TensorInfo[dim0][INPUT_STRIDE]];
+		}
+
+		bool DimEnd1 = false;
+		if (InputCoord1 == TensorInfo[dim1][INPUT_SIZE] - 1)
+		{
+			x010 = x000;
+			x011 = x001;
+			DimEnd1 = true;
+		}
+		else
+		{
+			x010 = Input[InputIndex + TensorInfo[dim1][INPUT_STRIDE]];
+			x011 = DimEnd0 ? x010 : Input[InputIndex + TensorInfo[dim0][INPUT_STRIDE] + TensorInfo[dim1][INPUT_STRIDE]];
+		}
+
+		if (InputCoord2 == TensorInfo[dim2][INPUT_SIZE] - 1)
+		{
+			x100 = x000;
+			x101 = x001;
+			x110 = x010;
+			x111 = x011;
+		}
+		else
+		{
+			x100 = Input[InputIndex + 1];
+			x101 = DimEnd0 ? x100 : Input[InputIndex + TensorInfo[dim0][INPUT_STRIDE] + 1];
+			x110 = DimEnd1 ? x100 : Input[InputIndex + TensorInfo[dim1][INPUT_STRIDE] + 1];
+			x111 = DimEnd0 && DimEnd1 ? x100 : DimEnd0 ? x110 : DimEnd1 ? x101 : Input[InputIndex + TensorInfo[dim0][INPUT_STRIDE] + TensorInfo[dim1][INPUT_STRIDE] + 1];
+		}
+
+		float y00 = x000 + (float)InputOffset0 * (x001 - x000) / (float)TensorInfo[dim0][SCALES_DIV];
+		float y01 = x010 + (float)InputOffset0 * (x011 - x010) / (float)TensorInfo[dim0][SCALES_DIV];
+		float y10 = x100 + (float)InputOffset0 * (x101 - x100) / (float)TensorInfo[dim0][SCALES_DIV];
+		float y11 = x110 + (float)InputOffset0 * (x111 - x110) / (float)TensorInfo[dim0][SCALES_DIV];
+
+		float z0 = y00 + (float)InputOffset1 * (y01 - y00) / (float)TensorInfo[dim1][SCALES_DIV];
+		float z1 = y10 + (float)InputOffset1 * (y11 - y10) / (float)TensorInfo[dim1][SCALES_DIV];
+
+		Output[OutputIndex] = z0 + (float)InputOffset2 * (z1 - z0) / (float)TensorInfo[dim2][SCALES_DIV];
 	}
 }
+
+#endif
\ No newline at end of file
diff --git a/Engine/Plugins/Experimental/NNERuntimeRDG/Source/NNERuntimeRDG/Internal/NNERuntimeRDGAttributes.h b/Engine/Plugins/Experimental/NNERuntimeRDG/Source/NNERuntimeRDG/Internal/NNERuntimeRDGAttributes.h
new file mode 100644
--- /dev/null
+++ b/Engine/Plugins/Experimental/NNERuntimeRDG/Source/NNERuntimeRDG/Internal/NNERuntimeRDGAttributes.h
@@ -0,0 +1,21 @@
+// Copyright Epic Games, Inc. All Rights Reserved.
+
+#pragma once
+
+#include "HAL/Platform.h"
+
+namespace UE::NNERuntimeRDG::Private
+{
+
+namespace AttrName
+{
+	static constexpr auto* Mode = TEXT("mode");
+}
+
+namespace AttrValue
+{
+	static constexpr auto* Nearest = TEXT("nearest");
+	static constexpr auto* Linear = TEXT("linear");
+}
+
+} // UE::NNERuntimeRDG::Private
\ No newline at end of file
diff --git a/Engine/Plugins/Experimental/NNERuntimeRDG/Source/NNERuntimeRDG/Private/Hlsl/NNERuntimeRDGUpsample.cpp b/Engine/Plugins/Experimental/NNERuntimeRDG/Source/NNERuntimeRDG/Private/Hlsl/NNERuntimeRDGUpsample.cpp
--- a/Engine/Plugins/Experimental/NNERuntimeRDG/Source/NNERuntimeRDG/Private/Hlsl/NNERuntimeRDGUpsample.cpp
+++ b/Engine/Plugins/Experimental/NNERuntimeRDG/Source/NNERuntimeRDG/Private/Hlsl/NNERuntimeRDGUpsample.cpp
@@ -1,10 +1,11 @@
 // Copyright Epic Games, Inc. All Rights Reserved.
 
 #include "NNERuntimeRDGUpsample.h"
+#include "NNERuntimeRDGAttributes.h"
 #include "NNEHlslShadersUpsampleCS.h"
 #include "NNERuntimeRDGHlslHelper.h"
-#include "NNETypes.h"
 #include "NNETensor.h"
+#include "NNETypes.h"
 
 namespace UE::NNERuntimeRDG::Private::Hlsl
 {
@@ -20,6 +21,10 @@ namespace UE::NNERuntimeRDG::Private::Hlsl
 		FUpsample() {}
 		virtual ~FUpsample() = default;
 
+	private:
+
+		FString Mode = AttrValue::Nearest;
+
 	public:
 
 		virtual int PrepareOutputs(TConstArrayView<NNE::Internal::FTensorRef> InputTensors, TArrayView<NNE::Internal::FTensorRef> OutputTensors) const override
@@ -27,7 +32,7 @@ namespace UE::NNERuntimeRDG::Private::Hlsl
 			check(InputTensors.Num() == 2);
 			check(OutputTensors.Num() == 1);
 
-			const NNE::Internal::FTensor& X = *InputTensors[0];
+			const NNE::Internal::FTensor& Input = *InputTensors[0];
 			const NNE::Internal::FTensor& Scales = *InputTensors[1];
 			
 			if (!Scales.HasPreparedData())
@@ -38,17 +43,32 @@ namespace UE::NNERuntimeRDG::Private::Hlsl
 
 			TConstArrayView<float> ScalesData = Scales.GetPreparedData<float>();
 
-			if (ScalesData.Num() != X.GetShape().Rank())
+			if (ScalesData.Num() != Input.GetShape().Rank())
 			{
-				UE_LOG(LogNNE, Warning, TEXT("Upsample input 'Scale' (name: %s) have %d elements. While it should be the same as the rank of input 'X' (name : %s) witch is %d"), *Scales.GetName(), ScalesData.Num(), *X.GetName(), X.GetShape().Rank());
+				UE_LOG(LogNNE, Warning, TEXT("Upsample input 'Scale' (name: %s) have %d elements. While it should be the same as the rank of input 'X' (name : %s) witch is %d"), *Scales.GetName(), ScalesData.Num(), *Input.GetName(), Input.GetShape().Rank());
 				return -1;
 			}
 
+			for (int32 i = 0; i < ScalesData.Num(); i++)
+			{
+				if (ScalesData[i] < 1.0f)
+				{
+					UE_LOG(LogNNE, Warning, TEXT("Upsample input 'Scale' takes values greater than or equal to 1."));
+					return -1;
+				}
+
+				if (Mode.Equals(AttrValue::Linear) && i < ScalesData.Num() - 3 && ScalesData[i] > 1.0f)
+				{
+					UE_LOG(LogNNE, Warning, TEXT("Upsample in 'Linear mode' only support up to trilinear interpolation, meaning input 'Scale' values for the outermost dimensions (Rank - 3) are 1."));
+					return -1;
+				}
+			}
+
 			TArray<uint32> OutputShapeData;
-			for (int32 i = 0; i < X.GetShape().Rank(); ++i)
+			for (int32 i = 0; i < Input.GetShape().Rank(); ++i)
 			{
 				
-				OutputShapeData.Emplace(FMath::FloorToInt32(X.GetShape().GetData()[i] * ScalesData[i]));
+				OutputShapeData.Emplace(FMath::FloorToInt32(Input.GetShape().GetData()[i] * ScalesData[i]));
 			}
 
 			NNE::FTensorShape OutputShape = NNE::FTensorShape::Make(OutputShapeData);
@@ -62,6 +82,8 @@ namespace UE::NNERuntimeRDG::Private::Hlsl
 			check(InputTensorDescs.Num() == 2);
 			check(OutputTensorDescs.Num() == 1);
 
+			Mode = Attributes.GetValueOrDefault<FString>(AttrName::Mode, Mode);
+
 			return true;
 		}
 
@@ -93,13 +115,20 @@ namespace UE::NNERuntimeRDG::Private::Hlsl
 			FillTensorStrideShaderParameters(Input, Params->TensorInfo, 0);
 			FillTensorStrideShaderParameters(Output, Params->TensorInfo, 1);
 			FillTensorSizeShaderParameters(Input, Params->TensorInfo, 2);
-			FillTensorSizeShaderParameters(Output, Params->TensorInfo, 3);
+			for (int32 i = 0; i < Input.GetShape().Rank(); ++i)
+			{
+				Params->TensorInfo[i][3] = FMath::CeilToInt32(ScalesData[i]);
+			}
 			Params->Num = Output.GetVolume();
 			Params->ThreadCountX = ThreadGroupCount.X * FUpsampleConstants::NUM_GROUP_THREADS;
 
-			FUpsampleCS::FPermutationDomain PermutationVector;
+			const EUpsampleMode UpsampleMode = Mode.Equals(AttrValue::Linear) ?
+				((ScalesData.Num() < 3 || ScalesData[ScalesData.Num() - 3] == 1.0f) ? EUpsampleMode::Bilinear : EUpsampleMode::Trilinear) :
+				EUpsampleMode::Nearest;
 
+			FUpsampleCS::FPermutationDomain PermutationVector;
 			PermutationVector.Set<FUpsampleCS::FUpsampleNumDimensions>(Output.GetShape().Rank());
+			PermutationVector.Set<FUpsampleCS::FUpsampleMode>(UpsampleMode);
 
 			TShaderMapRef<FUpsampleCS> ComputeShader(GetGlobalShaderMap(GMaxRHIFeatureLevel), PermutationVector);
 
@@ -121,13 +150,13 @@ namespace UE::NNERuntimeRDG::Private::Hlsl
 		bool bIsValid = true;
 
 		FAttributeValidator AttributeValidator;
-		AttributeValidator.AddOptional(TEXT("mode"), ENNEAttributeDataType::String);
+		AttributeValidator.AddOptional(AttrName::Mode, ENNEAttributeDataType::String);
 		bIsValid &= AttributeValidator.Validate(AttributeMap);
 		
-		FString Mode = AttributeMap.GetValueOrDefault<FString>(TEXT("mode"), TEXT("nearest"));
-		if (!Mode.Equals(TEXT("nearest")))
+		FString Mode = AttributeMap.GetValueOrDefault<FString>(AttrName::Mode, AttrValue::Nearest);
+		if (!Mode.Equals(AttrValue::Nearest) && !Mode.Equals(AttrValue::Linear))
 		{
-			UE_LOG(LogNNE, Warning, TEXT("Upsample HLSL operator only supports nearest mode for now"));
+			UE_LOG(LogNNE, Warning, TEXT("Upsample HLSL operator only supports %s or %s as value for attribute %s"), AttrValue::Nearest, AttrValue::Linear, AttrName::Mode);
 			return false;
 		}
 
