[NNE] Verify that FReadbackManager works on all platforms and simplify the API
Problem
- FReadbackManager utility class uses FStagingBufferRHI
- Not all platforms have a GPU fence, make sure that utility class is compatible by all RHIs
- Simplify the API to be inline with EnqueueReadback() from RDG utils

Solution
- Use RHICmdList.BlockUntilGPUIdle() which is compatible with all RHIs
- Refactor FReadbackManager to use the FRHIGPUBufferReadback (it has support for multiple GPUs) instead of using FStagingBufferRHI directly
- Simplify the readback API to have EnqueueReadbacks() and Wait()

Tests
- NNEIncubatorEditor Win64
- NNEIncubator standalone
- Non-unity compile

#rnx

diff --git a/Engine/Plugins/Experimental/NNERuntimeRDG/Source/NNERuntimeRDG/Internal/NNERenderGraphUtils.h b/Engine/Plugins/Experimental/NNERuntimeRDG/Source/NNERuntimeRDG/Internal/NNERenderGraphUtils.h
--- a/Engine/Plugins/Experimental/NNERuntimeRDG/Source/NNERuntimeRDG/Internal/NNERenderGraphUtils.h
+++ b/Engine/Plugins/Experimental/NNERuntimeRDG/Source/NNERuntimeRDG/Internal/NNERenderGraphUtils.h
@@ -5,6 +5,7 @@
 #include "NNERuntimeCPU.h"
 #include "NNERuntimeRDG.h"
 #include "RenderGraph.h"
+#include "RHIGPUReadback.h"
 
 #if UE_TRACE_ENABLED
 	#define NNE_TRACE_EVENT_SCOPED(Name) SCOPED_NAMED_EVENT_TEXT(#Name, FColor::Green)
@@ -15,23 +16,30 @@
 namespace UE::NNE::Internal
 {
 
+BEGIN_SHADER_PARAMETER_STRUCT(FReadbackPassParameters, )
+	RDG_BUFFER_ACCESS_ARRAY(Buffers)
+END_SHADER_PARAMETER_STRUCT()
+
 /*
 * Utility class to manage GPU -> CPU readbacks
+* The implementation follows single producer single consumer, 
+* every call to EnqueueReadbacks() on render thread has corresponding Wait() on game thread.
+* Don't enqueue multiple EnqueueReadbacks() while waiting on game thread, make sure that 
+* all readbacks are processed (by calling Wait()), before calling EnqueueReadbacks() again.
 */
 class FReadbackManager
 {
 	struct FReadback
 	{
-		FStagingBufferRHIRef	StagingBuffer;
-		FGPUFenceRHIRef			Fence;
-		FRDGBufferRef			Buffer;
+		FRHIGPUBufferReadback	BufferReadback;
 		void*					DstData;
 		uint32					NumBytes;
 	};
 
-	BEGIN_SHADER_PARAMETER_STRUCT(FReadbackPassParameters, )
-		RDG_BUFFER_ACCESS_ARRAY(ReadbackBuffers)
-	END_SHADER_PARAMETER_STRUCT()
+	using ReadbackArray = TArray<FReadback, TInlineAllocator<16>>;
+
+	ReadbackArray	Readbacks;
+	FEvent*			Signal = nullptr;
 
 public:
 
@@ -48,159 +56,96 @@ public:
 		Signal = nullptr;
 	}
 
-	// Note: call on the game thread
-	bool Init(int32 InNumReadbacks)
-	{
-		NNE_TRACE_EVENT_SCOPED(NNE_FTensorReadback_Init);
-
-		check(IsInGameThread());
-		MaxNumReadbacks = InNumReadbacks;
-		Signal->Reset();
-
-		return true;
-	}
-
-	void BeginReadbacks_RenderThread(FRDGBuilder& RDGBuilder)
+	bool EnqueueReadbacks(FRDGBuilder& RDGBuilder, TConstArrayView<FTensorBindingRDG> InBindingsRDG, TConstArrayView<FTensorBindingCPU> InBindingsCPU)
 	{
-		NNE_TRACE_EVENT_SCOPED(NNE_FTensorReadback_BeginReadbacks);
-
-		check(IsInRenderingThread());
-		check(Readbacks.IsEmpty());
-		check(MaxNumReadbacks > 0);
-
-		NumAddedReadbacks = 0;
-		NumProcessedReadbacks = 0;
-		Readbacks.SetNum(MaxNumReadbacks);
-		
-		ReadbackParams = RDGBuilder.AllocParameters<FReadbackPassParameters>();
-
-		for (FReadback& Readback : Readbacks)
+		if (!Readbacks.IsEmpty())
 		{
-			Readback.StagingBuffer = RHICreateStagingBuffer();
-			Readback.StagingBuffer->DisableLifetimeExtension();
-
-			Readback.Fence = RHICreateGPUFence(TEXT("FReadbackManager_ReadbackFence"));
-			Readback.Fence->DisableLifetimeExtension();
-			Readback.Fence->Clear();
-
-			Readback.Buffer = nullptr;
-			Readback.DstData = nullptr;
-			Readback.NumBytes = 0;
+			UE_LOG(LogNNE, Error, TEXT("FReadbackManager:Unprocessed readbacks detected"));
+			return false;
 		}
-	}
 
-	void AddReadbacks_RenderThread(TConstArrayView<FTensorBindingRDG> InBindingsRDG, TConstArrayView<FTensorBindingCPU> InBindingsCPU)
-	{
+		Signal->Reset();
+
 		NNE_TRACE_EVENT_SCOPED(NNE_FTensorReadback_AddReadbacks_RT);
 
 		check(IsInRenderingThread());
 		check(InBindingsRDG.Num() == InBindingsCPU.Num());
-		check(InBindingsRDG.Num() <= MaxNumReadbacks);
 
 		if (InBindingsRDG.Num() != InBindingsCPU.Num())
 		{
-			UE_LOG(LogNNE, Error, TEXT("FReadbackManager:Number of bindings need to be same"));
-			return;
+			UE_LOG(LogNNE, Error, TEXT("FReadbackManager:Number of bindings needs to be same"));
+			return false;
 		}
 
-		if (InBindingsRDG.Num() > MaxNumReadbacks)
-		{
-			UE_LOG(LogNNE, Error, TEXT("FReadbackManager:Number of bindings is larger than the maximum number of readbacks provided by the Init()"));
-			return;
-		}
+		FReadbackPassParameters* ReadbackParams = RDGBuilder.AllocParameters<FReadbackPassParameters>();
 
 		for (int32 Idx = 0; Idx < InBindingsRDG.Num(); ++Idx)
 		{
-			ReadbackParams->ReadbackBuffers.Emplace(InBindingsRDG[Idx].Buffer, ERHIAccess::CopySrc);
-
-			FReadback& Readback = Readbacks[NumAddedReadbacks];
-			Readback.Buffer = InBindingsRDG[Idx].Buffer;
-			Readback.DstData = InBindingsCPU[Idx].Data;
-			Readback.NumBytes = InBindingsCPU[Idx].SizeInBytes;
-
-			NumAddedReadbacks++;
+			ReadbackParams->Buffers.Emplace(InBindingsRDG[Idx].Buffer, ERHIAccess::CopySrc);
 		}
-	}
-
-	void EndReadbacks_RenderThread(FRDGBuilder& RDGBuilder)
-	{
-		check(IsInRenderingThread());
 
 		RDGBuilder.AddPass(
-			RDG_EVENT_NAME("FReadbackManager_ProcessReadbacks"),
+			RDG_EVENT_NAME("FReadbackManager_EnqueueReadbacks"),
 			ReadbackParams,
 			ERDGPassFlags::Readback | ERDGPassFlags::NeverCull,
-			[this](FRHICommandListImmediate& RHICmdList)
+			[this, ReadbackParams, InBindingsCPU](FRHICommandListImmediate& RHICmdList)
 			{
-				NNE_TRACE_EVENT_SCOPED(NNE_FTensorReadback_EndReadbacks_RT);
-				
-				for (FReadback& Readback : Readbacks)
+				for (int32 Idx = 0; Idx < ReadbackParams->Buffers.Num(); ++Idx)
 				{
-					Readback.Buffer->MarkResourceAsUsed();
-
-					RHICmdList.CopyToStagingBuffer(Readback.Buffer->GetRHI(), Readback.StagingBuffer, 0, Readback.NumBytes);					
-					RHICmdList.WriteGPUFence(Readback.Fence);
+					FReadback& Readback = Readbacks.Add_GetRef( 
+						{
+							FRHIGPUBufferReadback(*FString::Printf(TEXT("FReadbackManager_Readback_%d"), Idx)),
+							InBindingsCPU[Idx].Data, 
+							static_cast<uint32>(InBindingsCPU[Idx].SizeInBytes) 
+						}
+					);
+				
+					Readback.BufferReadback.EnqueueCopy(RHICmdList, ReadbackParams->Buffers[Idx]->GetRHI(), Readback.NumBytes);
 				}
+			}
+		);
+
+		return true;
+	}
+
+	// Note: call on the game thread
+	bool Wait()
+	{
+		check(IsInGameThread());
+
+		NNE_TRACE_EVENT_SCOPED(NNE_FTensorReadback_Wait);
+
+		ENQUEUE_RENDER_COMMAND(NNE_FReadbackManager_Wait)
+		(
+			[this](FRHICommandListImmediate& RHICmdList)
+			{
+				NNE_TRACE_EVENT_SCOPED(NNE_FTensorReadback_Wait_RT);
 
 				RHICmdList.BlockUntilGPUIdle();
 
 				for (FReadback& Readback : Readbacks)
 				{
-					const void* SrcData = RHICmdList.LockStagingBuffer(Readback.StagingBuffer, Readback.Fence.GetReference(), 0, Readback.NumBytes);
+					const void* SrcData = Readback.BufferReadback.Lock(Readback.NumBytes);
 					check(SrcData);
 
 					if (SrcData)
 					{
 						FMemory::Memcpy(Readback.DstData, SrcData, Readback.NumBytes);
-						RHICmdList.UnlockStagingBuffer(Readback.StagingBuffer);
-
-						Readback.StagingBuffer = nullptr;
-						Readback.Fence = nullptr;
+						Readback.BufferReadback.Unlock();
 					}
-
-					++NumProcessedReadbacks;
 				}
-				
-				// Clean-up resources (staging buffers, fences)
-				RHICmdList.ImmediateFlush(EImmediateFlushType::FlushRHIThreadFlushResources);
 
+				Readbacks.Reset();
 				Signal->Trigger();
+
+				// Clean-up resources (staging buffers, fences)
+				RHICmdList.ImmediateFlush(EImmediateFlushType::FlushRHIThreadFlushResources);
 			}
 		);
-	}
-
-	// Note: call on the game thread
-	void Wait()
-	{
-		check(IsInGameThread());
-
-		NNE_TRACE_EVENT_SCOPED(NNE_FTensorReadback_Wait);
 
-		check(MaxNumReadbacks > 0);
 		Signal->Wait();
-
-		while (NumProcessedReadbacks < MaxNumReadbacks)
-		{
-			FPlatformProcess::Sleep(0.f);
-		}
-
-		NumAddedReadbacks = 0;
-		NumProcessedReadbacks = 0;
-		MaxNumReadbacks = 0;
-		Readbacks.Empty();
+		return Readbacks.IsEmpty();
 	}
-
-
-private:
-
-	using ReadbackArray = TArray<FReadback, TInlineAllocator<16>>;
-	
-	ReadbackArray				Readbacks;
-	FReadbackPassParameters*	ReadbackParams = nullptr;
-	std::atomic<int32>			MaxNumReadbacks;
-	std::atomic<int32>			NumAddedReadbacks;
-	std::atomic<int32>			NumProcessedReadbacks;
-	FEvent*						Signal = nullptr;
 };
 
 } // namespace UE::NNE
