[NNE] NNERuntimeRDG PrepareOutput remove constness + Logsoftmax operator for HLSL

Problems:
* PrepareOutput as const does not scale as it is also responsible to prepare the operator when this one need to look for data in const tensors.
* Logsoftmax was missing for opset 9

Solutions:
* Remove constness from PrepareOutput
* Add logsoftmax operator for hlsl

Test: Win64 Editor: smoke tests
Test: Win64 Standalone : all tests

#rnx

#virtualized

diff --git a/Engine/Plugins/Experimental/NNERuntimeRDG/Source/NNEHlslShaders/Internal/NNEHlslShadersSoftmaxCS.h b/Engine/Plugins/Experimental/NNERuntimeRDG/Source/NNEHlslShaders/Internal/NNEHlslShadersSoftmaxCS.h
--- a/Engine/Plugins/Experimental/NNERuntimeRDG/Source/NNEHlslShaders/Internal/NNEHlslShadersSoftmaxCS.h
+++ b/Engine/Plugins/Experimental/NNERuntimeRDG/Source/NNEHlslShaders/Internal/NNEHlslShadersSoftmaxCS.h
@@ -10,6 +10,13 @@ namespace UE::NNE::Internal { class FTensor; }
 
 namespace UE::NNEHlslShaders::Internal
 {
+	enum class ESoftmaxOperatorType : uint8
+	{
+		SOFTMAX,
+		LOG_SOFTMAX,
+		MAX
+	};
+
 	class FSoftmaxConstants
 	{
 	public:
@@ -21,6 +28,9 @@ namespace UE::NNEHlslShaders::Internal
 		DECLARE_GLOBAL_SHADER(TSoftmaxCS);
 		SHADER_USE_PARAMETER_STRUCT(TSoftmaxCS, FHlslShaderBase)
 
+		class FSoftmaxType : SHADER_PERMUTATION_ENUM_CLASS("SOFTMAX_OPERATOR_TYPE", ESoftmaxOperatorType);
+		using FPermutationDomain = TShaderPermutationDomain<FSoftmaxType>;
+
 	public:
 
 		BEGIN_SHADER_PARAMETER_STRUCT(FParameters, )
diff --git a/Engine/Plugins/Experimental/NNERuntimeRDG/Source/NNEHlslShaders/Shaders/NNEHlslShadersSoftmax.usf b/Engine/Plugins/Experimental/NNERuntimeRDG/Source/NNEHlslShaders/Shaders/NNEHlslShadersSoftmax.usf
--- a/Engine/Plugins/Experimental/NNERuntimeRDG/Source/NNEHlslShaders/Shaders/NNEHlslShadersSoftmax.usf
+++ b/Engine/Plugins/Experimental/NNERuntimeRDG/Source/NNEHlslShaders/Shaders/NNEHlslShadersSoftmax.usf
@@ -14,6 +14,9 @@ uint Num;
 uint ThreadCountX;
 uint AxisSize;
 
+#define SOFTMAX_TYPE 0
+#define LOG_SOFTMAX_TYPE 1
+
 [numthreads(THREADGROUP_SIZE_X, 1, 1)]
 void Softmax(in const uint3 DispatchThreadID : SV_DispatchThreadID)
 {
@@ -22,6 +25,12 @@ void Softmax(in const uint3 DispatchThreadID : SV_DispatchThreadID)
 	{
 		uint SumExpIndex = Index/AxisSize;
 		WORK_TYPE SumExp = READ(InputSumExp[SumExpIndex]);
-		Output[Index] = WRITE(exp(READ(Input[Index])) / SumExp);
+		WORK_TYPE Result = WRITE(exp(READ(Input[Index])) / SumExp);
+
+		#if (SOFTMAX_OPERATOR_TYPE == LOG_SOFTMAX_TYPE)
+			Result = log(Result);
+		#endif
+
+		Output[Index] = Result;
 	}
 }
\ No newline at end of file
diff --git a/Engine/Plugins/Experimental/NNERuntimeRDG/Source/NNERuntimeRDG/Private/Dml/NNEDmlOperator.h b/Engine/Plugins/Experimental/NNERuntimeRDG/Source/NNERuntimeRDG/Private/Dml/NNEDmlOperator.h
--- a/Engine/Plugins/Experimental/NNERuntimeRDG/Source/NNERuntimeRDG/Private/Dml/NNEDmlOperator.h
+++ b/Engine/Plugins/Experimental/NNERuntimeRDG/Source/NNERuntimeRDG/Private/Dml/NNEDmlOperator.h
@@ -336,7 +336,7 @@ public:
 	/**
 	* Evaluate tensor shapes
 	*/
-	virtual int PrepareOutputs(TConstArrayView<NNE::Internal::FTensorRef> InputTensors, TArrayView<NNE::Internal::FTensorRef> OutputTensors) const = 0;
+	virtual int PrepareOutputs(TConstArrayView<NNE::Internal::FTensorRef> InputTensors, TArrayView<NNE::Internal::FTensorRef> OutputTensors) = 0;
 
 	/**
 	* Once the tensors have concrete shapes we can create the instance of the operator. The concrete shapes are evaluated in PrepareOutputs()
diff --git a/Engine/Plugins/Experimental/NNERuntimeRDG/Source/NNERuntimeRDG/Private/Dml/NNEDmlOperatorActivations.cpp b/Engine/Plugins/Experimental/NNERuntimeRDG/Source/NNERuntimeRDG/Private/Dml/NNEDmlOperatorActivations.cpp
--- a/Engine/Plugins/Experimental/NNERuntimeRDG/Source/NNERuntimeRDG/Private/Dml/NNEDmlOperatorActivations.cpp
+++ b/Engine/Plugins/Experimental/NNERuntimeRDG/Source/NNERuntimeRDG/Private/Dml/NNEDmlOperatorActivations.cpp
@@ -81,7 +81,7 @@ public:
 		return true;
 	}
 
-	virtual int PrepareOutputs(TConstArrayView<NNE::Internal::FTensorRef> InputTensors, TArrayView<NNE::Internal::FTensorRef> OutputTensors) const override
+	virtual int PrepareOutputs(TConstArrayView<NNE::Internal::FTensorRef> InputTensors, TArrayView<NNE::Internal::FTensorRef> OutputTensors) override
 	{
 		check(InputTensors.Num() == NumAllowedInputTensors);
 		check(OutputTensors.Num() == NumAllowedOutputTensors);
@@ -247,7 +247,7 @@ public:
 		return true;
 	}
 
-	virtual int PrepareOutputs(TConstArrayView<NNE::Internal::FTensorRef> InputTensors, TArrayView<NNE::Internal::FTensorRef> OutputTensors) const override
+	virtual int PrepareOutputs(TConstArrayView<NNE::Internal::FTensorRef> InputTensors, TArrayView<NNE::Internal::FTensorRef> OutputTensors) override
 	{
 		check(InputTensors.Num() == NumAllowedInputTensors);
 		check(OutputTensors.Num() == NumAllowedOutputTensors);
diff --git a/Engine/Plugins/Experimental/NNERuntimeRDG/Source/NNERuntimeRDG/Private/Dml/NNEDmlOperatorBatchNormalization.cpp b/Engine/Plugins/Experimental/NNERuntimeRDG/Source/NNERuntimeRDG/Private/Dml/NNEDmlOperatorBatchNormalization.cpp
--- a/Engine/Plugins/Experimental/NNERuntimeRDG/Source/NNERuntimeRDG/Private/Dml/NNEDmlOperatorBatchNormalization.cpp
+++ b/Engine/Plugins/Experimental/NNERuntimeRDG/Source/NNERuntimeRDG/Private/Dml/NNEDmlOperatorBatchNormalization.cpp
@@ -93,7 +93,7 @@ public:
 		return true;
 	}
 
-	virtual int PrepareOutputs(TConstArrayView<NNE::Internal::FTensorRef> InputTensors, TArrayView<NNE::Internal::FTensorRef> OutputTensors) const override
+	virtual int PrepareOutputs(TConstArrayView<NNE::Internal::FTensorRef> InputTensors, TArrayView<NNE::Internal::FTensorRef> OutputTensors) override
 	{
 		check(InputTensors.Num() == Count);
 		check(OutputTensors.Num() == 1);
diff --git a/Engine/Plugins/Experimental/NNERuntimeRDG/Source/NNERuntimeRDG/Private/Dml/NNEDmlOperatorCast.cpp b/Engine/Plugins/Experimental/NNERuntimeRDG/Source/NNERuntimeRDG/Private/Dml/NNEDmlOperatorCast.cpp
--- a/Engine/Plugins/Experimental/NNERuntimeRDG/Source/NNERuntimeRDG/Private/Dml/NNEDmlOperatorCast.cpp
+++ b/Engine/Plugins/Experimental/NNERuntimeRDG/Source/NNERuntimeRDG/Private/Dml/NNEDmlOperatorCast.cpp
@@ -73,7 +73,7 @@ public:
 		return true;
 	}
 
-	virtual int PrepareOutputs(TConstArrayView<NNE::Internal::FTensorRef> InputTensors, TArrayView<NNE::Internal::FTensorRef> OutputTensors) const override
+	virtual int PrepareOutputs(TConstArrayView<NNE::Internal::FTensorRef> InputTensors, TArrayView<NNE::Internal::FTensorRef> OutputTensors) override
 	{
 		check(InputTensors.Num() == NumAllowedInputTensors);
 		check(OutputTensors.Num() == NumAllowedOutputTensors);
diff --git a/Engine/Plugins/Experimental/NNERuntimeRDG/Source/NNERuntimeRDG/Private/Dml/NNEDmlOperatorConcat.cpp b/Engine/Plugins/Experimental/NNERuntimeRDG/Source/NNERuntimeRDG/Private/Dml/NNEDmlOperatorConcat.cpp
--- a/Engine/Plugins/Experimental/NNERuntimeRDG/Source/NNERuntimeRDG/Private/Dml/NNEDmlOperatorConcat.cpp
+++ b/Engine/Plugins/Experimental/NNERuntimeRDG/Source/NNERuntimeRDG/Private/Dml/NNEDmlOperatorConcat.cpp
@@ -103,7 +103,7 @@ public:
 		return true;
 	}
 
-	virtual int PrepareOutputs(TConstArrayView<NNE::Internal::FTensorRef> InputTensors, TArrayView<NNE::Internal::FTensorRef> OutputTensors) const override
+	virtual int PrepareOutputs(TConstArrayView<NNE::Internal::FTensorRef> InputTensors, TArrayView<NNE::Internal::FTensorRef> OutputTensors) override
 	{
 		TArray<uint32> OutputShape(InputTensors[0]->GetShape().GetData());
 		
diff --git a/Engine/Plugins/Experimental/NNERuntimeRDG/Source/NNERuntimeRDG/Private/Dml/NNEDmlOperatorConvolution.cpp b/Engine/Plugins/Experimental/NNERuntimeRDG/Source/NNERuntimeRDG/Private/Dml/NNEDmlOperatorConvolution.cpp
--- a/Engine/Plugins/Experimental/NNERuntimeRDG/Source/NNERuntimeRDG/Private/Dml/NNEDmlOperatorConvolution.cpp
+++ b/Engine/Plugins/Experimental/NNERuntimeRDG/Source/NNERuntimeRDG/Private/Dml/NNEDmlOperatorConvolution.cpp
@@ -170,7 +170,7 @@ public:
 		return true;
 	}
 
-	virtual int PrepareOutputs(TConstArrayView<NNE::Internal::FTensorRef> InputTensors, TArrayView<NNE::Internal::FTensorRef> OutputTensors) const override
+	virtual int PrepareOutputs(TConstArrayView<NNE::Internal::FTensorRef> InputTensors, TArrayView<NNE::Internal::FTensorRef> OutputTensors) override
 	{
 		NNE::Internal::FTensor& OutputTensor = *OutputTensors[0];
 
diff --git a/Engine/Plugins/Experimental/NNERuntimeRDG/Source/NNERuntimeRDG/Private/Dml/NNEDmlOperatorCumSum.cpp b/Engine/Plugins/Experimental/NNERuntimeRDG/Source/NNERuntimeRDG/Private/Dml/NNEDmlOperatorCumSum.cpp
--- a/Engine/Plugins/Experimental/NNERuntimeRDG/Source/NNERuntimeRDG/Private/Dml/NNEDmlOperatorCumSum.cpp
+++ b/Engine/Plugins/Experimental/NNERuntimeRDG/Source/NNERuntimeRDG/Private/Dml/NNEDmlOperatorCumSum.cpp
@@ -65,7 +65,7 @@ public:
 		return true;
 	}
 
-	virtual int PrepareOutputs(TConstArrayView<NNE::Internal::FTensorRef> InputTensors, TArrayView<NNE::Internal::FTensorRef> OutputTensors) const override
+	virtual int PrepareOutputs(TConstArrayView<NNE::Internal::FTensorRef> InputTensors, TArrayView<NNE::Internal::FTensorRef> OutputTensors) override
 	{
 		// axis tensor must be constant
 		check(InputTensors[1]->HasPreparedData());
diff --git a/Engine/Plugins/Experimental/NNERuntimeRDG/Source/NNERuntimeRDG/Private/Dml/NNEDmlOperatorDepthToSpace.cpp b/Engine/Plugins/Experimental/NNERuntimeRDG/Source/NNERuntimeRDG/Private/Dml/NNEDmlOperatorDepthToSpace.cpp
--- a/Engine/Plugins/Experimental/NNERuntimeRDG/Source/NNERuntimeRDG/Private/Dml/NNEDmlOperatorDepthToSpace.cpp
+++ b/Engine/Plugins/Experimental/NNERuntimeRDG/Source/NNERuntimeRDG/Private/Dml/NNEDmlOperatorDepthToSpace.cpp
@@ -86,7 +86,7 @@ public:
 		return true;
 	}
 
-	virtual int PrepareOutputs(TConstArrayView<NNE::Internal::FTensorRef> InputTensors, TArrayView<NNE::Internal::FTensorRef> OutputTensors) const override
+	virtual int PrepareOutputs(TConstArrayView<NNE::Internal::FTensorRef> InputTensors, TArrayView<NNE::Internal::FTensorRef> OutputTensors) override
 	{
 		TConstArrayView<uint32>		InputShape = InputTensors[0]->GetShape().GetData();
 		Util::FSmallUIntArray		OutputShape;
diff --git a/Engine/Plugins/Experimental/NNERuntimeRDG/Source/NNERuntimeRDG/Private/Dml/NNEDmlOperatorElementwise.cpp b/Engine/Plugins/Experimental/NNERuntimeRDG/Source/NNERuntimeRDG/Private/Dml/NNEDmlOperatorElementwise.cpp
--- a/Engine/Plugins/Experimental/NNERuntimeRDG/Source/NNERuntimeRDG/Private/Dml/NNEDmlOperatorElementwise.cpp
+++ b/Engine/Plugins/Experimental/NNERuntimeRDG/Source/NNERuntimeRDG/Private/Dml/NNEDmlOperatorElementwise.cpp
@@ -82,7 +82,7 @@ public:
 		return true;
 	}
 
-	virtual int PrepareOutputs(TConstArrayView<NNE::Internal::FTensorRef> InputTensors, TArrayView<NNE::Internal::FTensorRef> OutputTensors) const override
+	virtual int PrepareOutputs(TConstArrayView<NNE::Internal::FTensorRef> InputTensors, TArrayView<NNE::Internal::FTensorRef> OutputTensors) override
 	{
 		check(InputTensors.Num() == NumAllowedInputTensors);
 		check(OutputTensors.Num() == NumAllowedOutputTensors);
@@ -189,7 +189,7 @@ public:
 		return true;
 	}
 
-	virtual int PrepareOutputs(TConstArrayView<NNE::Internal::FTensorRef> InputTensors, TArrayView<NNE::Internal::FTensorRef> OutputTensors) const override
+	virtual int PrepareOutputs(TConstArrayView<NNE::Internal::FTensorRef> InputTensors, TArrayView<NNE::Internal::FTensorRef> OutputTensors) override
 	{
 		check(InputTensors.Num() == NumAllowedInputTensors);
 		check(OutputTensors.Num() == NumAllowedOutputTensors);
diff --git a/Engine/Plugins/Experimental/NNERuntimeRDG/Source/NNERuntimeRDG/Private/Dml/NNEDmlOperatorFlatten.cpp b/Engine/Plugins/Experimental/NNERuntimeRDG/Source/NNERuntimeRDG/Private/Dml/NNEDmlOperatorFlatten.cpp
--- a/Engine/Plugins/Experimental/NNERuntimeRDG/Source/NNERuntimeRDG/Private/Dml/NNEDmlOperatorFlatten.cpp
+++ b/Engine/Plugins/Experimental/NNERuntimeRDG/Source/NNERuntimeRDG/Private/Dml/NNEDmlOperatorFlatten.cpp
@@ -70,7 +70,7 @@ public:
 		return true;
 	}
 
-	virtual int PrepareOutputs(TConstArrayView<NNE::Internal::FTensorRef> InputTensors, TArrayView<NNE::Internal::FTensorRef> OutputTensors) const override
+	virtual int PrepareOutputs(TConstArrayView<NNE::Internal::FTensorRef> InputTensors, TArrayView<NNE::Internal::FTensorRef> OutputTensors) override
 	{
 		TConstArrayView<uint32> InputShape = InputTensors[0]->GetShape().GetData();
 
diff --git a/Engine/Plugins/Experimental/NNERuntimeRDG/Source/NNERuntimeRDG/Private/Dml/NNEDmlOperatorGather.cpp b/Engine/Plugins/Experimental/NNERuntimeRDG/Source/NNERuntimeRDG/Private/Dml/NNEDmlOperatorGather.cpp
--- a/Engine/Plugins/Experimental/NNERuntimeRDG/Source/NNERuntimeRDG/Private/Dml/NNEDmlOperatorGather.cpp
+++ b/Engine/Plugins/Experimental/NNERuntimeRDG/Source/NNERuntimeRDG/Private/Dml/NNEDmlOperatorGather.cpp
@@ -91,7 +91,7 @@ public:
 		return true;
 	}
 
-	virtual int PrepareOutputs(TConstArrayView<NNE::Internal::FTensorRef> InputTensors, TArrayView<NNE::Internal::FTensorRef> OutputTensors) const override
+	virtual int PrepareOutputs(TConstArrayView<NNE::Internal::FTensorRef> InputTensors, TArrayView<NNE::Internal::FTensorRef> OutputTensors) override
 	{
 		const NNE::FTensorShape& InputShape = InputTensors[0]->GetShape();
 		const NNE::FTensorShape& IndicesShape = InputTensors[1]->GetShape();
diff --git a/Engine/Plugins/Experimental/NNERuntimeRDG/Source/NNERuntimeRDG/Private/Dml/NNEDmlOperatorGemm.cpp b/Engine/Plugins/Experimental/NNERuntimeRDG/Source/NNERuntimeRDG/Private/Dml/NNEDmlOperatorGemm.cpp
--- a/Engine/Plugins/Experimental/NNERuntimeRDG/Source/NNERuntimeRDG/Private/Dml/NNEDmlOperatorGemm.cpp
+++ b/Engine/Plugins/Experimental/NNERuntimeRDG/Source/NNERuntimeRDG/Private/Dml/NNEDmlOperatorGemm.cpp
@@ -100,7 +100,7 @@ public:
 		return true;
 	}
 
-	virtual int PrepareOutputs(TConstArrayView<NNE::Internal::FTensorRef> InputTensors, TArrayView<NNE::Internal::FTensorRef> OutputTensors) const override
+	virtual int PrepareOutputs(TConstArrayView<NNE::Internal::FTensorRef> InputTensors, TArrayView<NNE::Internal::FTensorRef> OutputTensors) override
 	{
 		check(InputTensors.Num() >= MinAllowedInputTensors && InputTensors.Num() <= MaxAllowedInputTensors);
 		check(OutputTensors.Num() == NumAllowedOutputTensors);
diff --git a/Engine/Plugins/Experimental/NNERuntimeRDG/Source/NNERuntimeRDG/Private/Dml/NNEDmlOperatorIdentity.cpp b/Engine/Plugins/Experimental/NNERuntimeRDG/Source/NNERuntimeRDG/Private/Dml/NNEDmlOperatorIdentity.cpp
--- a/Engine/Plugins/Experimental/NNERuntimeRDG/Source/NNERuntimeRDG/Private/Dml/NNEDmlOperatorIdentity.cpp
+++ b/Engine/Plugins/Experimental/NNERuntimeRDG/Source/NNERuntimeRDG/Private/Dml/NNEDmlOperatorIdentity.cpp
@@ -54,7 +54,7 @@ public:
 		return true;
 	}
 
-	virtual int PrepareOutputs(TConstArrayView<NNE::Internal::FTensorRef> InputTensors, TArrayView<NNE::Internal::FTensorRef> OutputTensors) const override
+	virtual int PrepareOutputs(TConstArrayView<NNE::Internal::FTensorRef> InputTensors, TArrayView<NNE::Internal::FTensorRef> OutputTensors) override
 	{
 		OutputTensors[0]->SetShape(InputTensors[0]->GetShape());
 		return 0;
diff --git a/Engine/Plugins/Experimental/NNERuntimeRDG/Source/NNERuntimeRDG/Private/Dml/NNEDmlOperatorInstanceNormalization.cpp b/Engine/Plugins/Experimental/NNERuntimeRDG/Source/NNERuntimeRDG/Private/Dml/NNEDmlOperatorInstanceNormalization.cpp
--- a/Engine/Plugins/Experimental/NNERuntimeRDG/Source/NNERuntimeRDG/Private/Dml/NNEDmlOperatorInstanceNormalization.cpp
+++ b/Engine/Plugins/Experimental/NNERuntimeRDG/Source/NNERuntimeRDG/Private/Dml/NNEDmlOperatorInstanceNormalization.cpp
@@ -70,7 +70,7 @@ public:
 		return true;
 	}
 
-	virtual int PrepareOutputs(TConstArrayView<NNE::Internal::FTensorRef> InputTensors, TArrayView<NNE::Internal::FTensorRef> OutputTensors) const override
+	virtual int PrepareOutputs(TConstArrayView<NNE::Internal::FTensorRef> InputTensors, TArrayView<NNE::Internal::FTensorRef> OutputTensors) override
 	{
 		OutputTensors[0]->SetShape(InputTensors[0]->GetShape());
 		return 0;
diff --git a/Engine/Plugins/Experimental/NNERuntimeRDG/Source/NNERuntimeRDG/Private/Dml/NNEDmlOperatorLRN.cpp b/Engine/Plugins/Experimental/NNERuntimeRDG/Source/NNERuntimeRDG/Private/Dml/NNEDmlOperatorLRN.cpp
--- a/Engine/Plugins/Experimental/NNERuntimeRDG/Source/NNERuntimeRDG/Private/Dml/NNEDmlOperatorLRN.cpp
+++ b/Engine/Plugins/Experimental/NNERuntimeRDG/Source/NNERuntimeRDG/Private/Dml/NNEDmlOperatorLRN.cpp
@@ -59,7 +59,7 @@ public:
 		return true;
 	}
 
-	virtual int PrepareOutputs(TConstArrayView<NNE::Internal::FTensorRef> InputTensors, TArrayView<NNE::Internal::FTensorRef> OutputTensors) const override
+	virtual int PrepareOutputs(TConstArrayView<NNE::Internal::FTensorRef> InputTensors, TArrayView<NNE::Internal::FTensorRef> OutputTensors) override
 	{
 		OutputTensors[0]->SetShape(InputTensors[0]->GetShape());
 		return 0;
diff --git a/Engine/Plugins/Experimental/NNERuntimeRDG/Source/NNERuntimeRDG/Private/Dml/NNEDmlOperatorLpNormalization.cpp b/Engine/Plugins/Experimental/NNERuntimeRDG/Source/NNERuntimeRDG/Private/Dml/NNEDmlOperatorLpNormalization.cpp
--- a/Engine/Plugins/Experimental/NNERuntimeRDG/Source/NNERuntimeRDG/Private/Dml/NNEDmlOperatorLpNormalization.cpp
+++ b/Engine/Plugins/Experimental/NNERuntimeRDG/Source/NNERuntimeRDG/Private/Dml/NNEDmlOperatorLpNormalization.cpp
@@ -60,7 +60,7 @@ public:
 		return true;
 	}
 
-	virtual int PrepareOutputs(TConstArrayView<NNE::Internal::FTensorRef> InputTensors, TArrayView<NNE::Internal::FTensorRef> OutputTensors) const override
+	virtual int PrepareOutputs(TConstArrayView<NNE::Internal::FTensorRef> InputTensors, TArrayView<NNE::Internal::FTensorRef> OutputTensors) override
 	{
 		OutputTensors[0]->SetShape(InputTensors[0]->GetShape());
 		return 0;
diff --git a/Engine/Plugins/Experimental/NNERuntimeRDG/Source/NNERuntimeRDG/Private/Dml/NNEDmlOperatorMatMul.cpp b/Engine/Plugins/Experimental/NNERuntimeRDG/Source/NNERuntimeRDG/Private/Dml/NNEDmlOperatorMatMul.cpp
--- a/Engine/Plugins/Experimental/NNERuntimeRDG/Source/NNERuntimeRDG/Private/Dml/NNEDmlOperatorMatMul.cpp
+++ b/Engine/Plugins/Experimental/NNERuntimeRDG/Source/NNERuntimeRDG/Private/Dml/NNEDmlOperatorMatMul.cpp
@@ -57,7 +57,7 @@ public:
 		return true;
 	}
 
-	virtual int PrepareOutputs(TConstArrayView<NNE::Internal::FTensorRef> InputTensors, TArrayView<NNE::Internal::FTensorRef> OutputTensors) const override
+	virtual int PrepareOutputs(TConstArrayView<NNE::Internal::FTensorRef> InputTensors, TArrayView<NNE::Internal::FTensorRef> OutputTensors) override
 	{
 		check(InputTensors.Num() == NumAllowedInputTensors);
 		check(OutputTensors.Num() == NumAllowedOutputTensors);
diff --git a/Engine/Plugins/Experimental/NNERuntimeRDG/Source/NNERuntimeRDG/Private/Dml/NNEDmlOperatorMaxUnpool.cpp b/Engine/Plugins/Experimental/NNERuntimeRDG/Source/NNERuntimeRDG/Private/Dml/NNEDmlOperatorMaxUnpool.cpp
--- a/Engine/Plugins/Experimental/NNERuntimeRDG/Source/NNERuntimeRDG/Private/Dml/NNEDmlOperatorMaxUnpool.cpp
+++ b/Engine/Plugins/Experimental/NNERuntimeRDG/Source/NNERuntimeRDG/Private/Dml/NNEDmlOperatorMaxUnpool.cpp
@@ -129,7 +129,7 @@ public:
 		return true;
 	}
 
-	virtual int PrepareOutputs(TConstArrayView<NNE::Internal::FTensorRef> InputTensors, TArrayView<NNE::Internal::FTensorRef> OutputTensors) const override
+	virtual int PrepareOutputs(TConstArrayView<NNE::Internal::FTensorRef> InputTensors, TArrayView<NNE::Internal::FTensorRef> OutputTensors) override
 	{
 		Args.Evaluate(InputTensors[0]->GetShape().GetData());
 
diff --git a/Engine/Plugins/Experimental/NNERuntimeRDG/Source/NNERuntimeRDG/Private/Dml/NNEDmlOperatorMeanVarianceNormalization.cpp b/Engine/Plugins/Experimental/NNERuntimeRDG/Source/NNERuntimeRDG/Private/Dml/NNEDmlOperatorMeanVarianceNormalization.cpp
--- a/Engine/Plugins/Experimental/NNERuntimeRDG/Source/NNERuntimeRDG/Private/Dml/NNEDmlOperatorMeanVarianceNormalization.cpp
+++ b/Engine/Plugins/Experimental/NNERuntimeRDG/Source/NNERuntimeRDG/Private/Dml/NNEDmlOperatorMeanVarianceNormalization.cpp
@@ -83,7 +83,7 @@ public:
 		return true;
 	}
 
-	virtual int PrepareOutputs(TConstArrayView<NNE::Internal::FTensorRef> InputTensors, TArrayView<NNE::Internal::FTensorRef> OutputTensors) const override
+	virtual int PrepareOutputs(TConstArrayView<NNE::Internal::FTensorRef> InputTensors, TArrayView<NNE::Internal::FTensorRef> OutputTensors) override
 	{
 		OutputTensors[0]->SetShape(InputTensors[0]->GetShape());
 		return 0;
diff --git a/Engine/Plugins/Experimental/NNERuntimeRDG/Source/NNERuntimeRDG/Private/Dml/NNEDmlOperatorPad.cpp b/Engine/Plugins/Experimental/NNERuntimeRDG/Source/NNERuntimeRDG/Private/Dml/NNEDmlOperatorPad.cpp
--- a/Engine/Plugins/Experimental/NNERuntimeRDG/Source/NNERuntimeRDG/Private/Dml/NNEDmlOperatorPad.cpp
+++ b/Engine/Plugins/Experimental/NNERuntimeRDG/Source/NNERuntimeRDG/Private/Dml/NNEDmlOperatorPad.cpp
@@ -103,7 +103,7 @@ public:
 		return true;
 	}
 
-	virtual int PrepareOutputs(TConstArrayView<NNE::Internal::FTensorRef> InputTensors, TArrayView<NNE::Internal::FTensorRef> OutputTensors) const override
+	virtual int PrepareOutputs(TConstArrayView<NNE::Internal::FTensorRef> InputTensors, TArrayView<NNE::Internal::FTensorRef> OutputTensors) override
 	{
 		if (InputTensors.Num() >= 2)
 		{
diff --git a/Engine/Plugins/Experimental/NNERuntimeRDG/Source/NNERuntimeRDG/Private/Dml/NNEDmlOperatorPool.cpp b/Engine/Plugins/Experimental/NNERuntimeRDG/Source/NNERuntimeRDG/Private/Dml/NNEDmlOperatorPool.cpp
--- a/Engine/Plugins/Experimental/NNERuntimeRDG/Source/NNERuntimeRDG/Private/Dml/NNEDmlOperatorPool.cpp
+++ b/Engine/Plugins/Experimental/NNERuntimeRDG/Source/NNERuntimeRDG/Private/Dml/NNEDmlOperatorPool.cpp
@@ -164,7 +164,7 @@ public:
 		return true;
 	}
 
-	virtual int PrepareOutputs(TConstArrayView<NNE::Internal::FTensorRef> InputTensors, TArrayView<NNE::Internal::FTensorRef> OutputTensors) const override
+	virtual int PrepareOutputs(TConstArrayView<NNE::Internal::FTensorRef> InputTensors, TArrayView<NNE::Internal::FTensorRef> OutputTensors) override
 	{
 		TConstArrayView<uint32> InputShape = InputTensors[0]->GetShape().GetData();
 
diff --git a/Engine/Plugins/Experimental/NNERuntimeRDG/Source/NNERuntimeRDG/Private/Dml/NNEDmlOperatorReduce.cpp b/Engine/Plugins/Experimental/NNERuntimeRDG/Source/NNERuntimeRDG/Private/Dml/NNEDmlOperatorReduce.cpp
--- a/Engine/Plugins/Experimental/NNERuntimeRDG/Source/NNERuntimeRDG/Private/Dml/NNEDmlOperatorReduce.cpp
+++ b/Engine/Plugins/Experimental/NNERuntimeRDG/Source/NNERuntimeRDG/Private/Dml/NNEDmlOperatorReduce.cpp
@@ -114,7 +114,7 @@ public:
 		return true;
 	}
 
-	virtual int PrepareOutputs(TConstArrayView<NNE::Internal::FTensorRef> InputTensors, TArrayView<NNE::Internal::FTensorRef> OutputTensors) const override
+	virtual int PrepareOutputs(TConstArrayView<NNE::Internal::FTensorRef> InputTensors, TArrayView<NNE::Internal::FTensorRef> OutputTensors) override
 	{
 		TConstArrayView<uint32> InputShape = InputTensors[0]->GetShape().GetData();
 		Util::FSmallUIntArray	OutputShape;
diff --git a/Engine/Plugins/Experimental/NNERuntimeRDG/Source/NNERuntimeRDG/Private/Dml/NNEDmlOperatorResample.cpp b/Engine/Plugins/Experimental/NNERuntimeRDG/Source/NNERuntimeRDG/Private/Dml/NNEDmlOperatorResample.cpp
--- a/Engine/Plugins/Experimental/NNERuntimeRDG/Source/NNERuntimeRDG/Private/Dml/NNEDmlOperatorResample.cpp
+++ b/Engine/Plugins/Experimental/NNERuntimeRDG/Source/NNERuntimeRDG/Private/Dml/NNEDmlOperatorResample.cpp
@@ -278,7 +278,7 @@ public:
 		return true;
 	}
 	
-	virtual int PrepareOutputs(TConstArrayView<NNE::Internal::FTensorRef> InputTensors, TArrayView<NNE::Internal::FTensorRef> OutputTensors) const override
+	virtual int PrepareOutputs(TConstArrayView<NNE::Internal::FTensorRef> InputTensors, TArrayView<NNE::Internal::FTensorRef> OutputTensors) override
 	{
 		const NNE::Internal::FTensor& InputTensor = *InputTensors[0];
 		const NNE::Internal::FTensor& ScaleTensor = (InputTensors.Num() == 2) ? *InputTensors[1] : *InputTensors[2]; // Upsample has scale at position 1, while Resize at 2
diff --git a/Engine/Plugins/Experimental/NNERuntimeRDG/Source/NNERuntimeRDG/Private/Dml/NNEDmlOperatorReshape.cpp b/Engine/Plugins/Experimental/NNERuntimeRDG/Source/NNERuntimeRDG/Private/Dml/NNEDmlOperatorReshape.cpp
--- a/Engine/Plugins/Experimental/NNERuntimeRDG/Source/NNERuntimeRDG/Private/Dml/NNEDmlOperatorReshape.cpp
+++ b/Engine/Plugins/Experimental/NNERuntimeRDG/Source/NNERuntimeRDG/Private/Dml/NNEDmlOperatorReshape.cpp
@@ -74,7 +74,7 @@ public:
 		return true;
 	}
 
-	virtual int PrepareOutputs(TConstArrayView<NNE::Internal::FTensorRef> InputTensors, TArrayView<NNE::Internal::FTensorRef> OutputTensors) const override
+	virtual int PrepareOutputs(TConstArrayView<NNE::Internal::FTensorRef> InputTensors, TArrayView<NNE::Internal::FTensorRef> OutputTensors) override
 	{
 		// Shape tensor must be constant!
 		check(InputTensors[1]->HasPreparedData());
diff --git a/Engine/Plugins/Experimental/NNERuntimeRDG/Source/NNERuntimeRDG/Private/Dml/NNEDmlOperatorSlice.cpp b/Engine/Plugins/Experimental/NNERuntimeRDG/Source/NNERuntimeRDG/Private/Dml/NNEDmlOperatorSlice.cpp
--- a/Engine/Plugins/Experimental/NNERuntimeRDG/Source/NNERuntimeRDG/Private/Dml/NNEDmlOperatorSlice.cpp
+++ b/Engine/Plugins/Experimental/NNERuntimeRDG/Source/NNERuntimeRDG/Private/Dml/NNEDmlOperatorSlice.cpp
@@ -234,7 +234,7 @@ public:
 		return true;
 	}
 
-	virtual int PrepareOutputs(TConstArrayView<NNE::Internal::FTensorRef> InputTensors, TArrayView<NNE::Internal::FTensorRef> OutputTensors) const override
+	virtual int PrepareOutputs(TConstArrayView<NNE::Internal::FTensorRef> InputTensors, TArrayView<NNE::Internal::FTensorRef> OutputTensors) override
 	{
 		switch (InputIndexDataType)
 		{
diff --git a/Engine/Plugins/Experimental/NNERuntimeRDG/Source/NNERuntimeRDG/Private/Dml/NNEDmlOperatorSpaceToDepth.cpp b/Engine/Plugins/Experimental/NNERuntimeRDG/Source/NNERuntimeRDG/Private/Dml/NNEDmlOperatorSpaceToDepth.cpp
--- a/Engine/Plugins/Experimental/NNERuntimeRDG/Source/NNERuntimeRDG/Private/Dml/NNEDmlOperatorSpaceToDepth.cpp
+++ b/Engine/Plugins/Experimental/NNERuntimeRDG/Source/NNERuntimeRDG/Private/Dml/NNEDmlOperatorSpaceToDepth.cpp
@@ -72,7 +72,7 @@ public:
 		return true;
 	}
 
-	virtual int PrepareOutputs(TConstArrayView<NNE::Internal::FTensorRef> InputTensors, TArrayView<NNE::Internal::FTensorRef> OutputTensors) const override
+	virtual int PrepareOutputs(TConstArrayView<NNE::Internal::FTensorRef> InputTensors, TArrayView<NNE::Internal::FTensorRef> OutputTensors) override
 	{
 		TConstArrayView<uint32>		InputShape = InputTensors[0]->GetShape().GetData();
 		Util::FSmallUIntArray		OutputShape;
diff --git a/Engine/Plugins/Experimental/NNERuntimeRDG/Source/NNERuntimeRDG/Private/Dml/NNEDmlOperatorSplit.cpp b/Engine/Plugins/Experimental/NNERuntimeRDG/Source/NNERuntimeRDG/Private/Dml/NNEDmlOperatorSplit.cpp
--- a/Engine/Plugins/Experimental/NNERuntimeRDG/Source/NNERuntimeRDG/Private/Dml/NNEDmlOperatorSplit.cpp
+++ b/Engine/Plugins/Experimental/NNERuntimeRDG/Source/NNERuntimeRDG/Private/Dml/NNEDmlOperatorSplit.cpp
@@ -85,7 +85,7 @@ public:
 		return true;
 	}
 
-	virtual int PrepareOutputs(TConstArrayView<NNE::Internal::FTensorRef> InputTensors, TArrayView<NNE::Internal::FTensorRef> OutputTensors) const override
+	virtual int PrepareOutputs(TConstArrayView<NNE::Internal::FTensorRef> InputTensors, TArrayView<NNE::Internal::FTensorRef> OutputTensors) override
 	{
 		const NNE::Internal::FTensor&	InputTensor = *InputTensors[0];
 		TConstArrayView<uint32>			InputShape = InputTensor.GetShape().GetData();
diff --git a/Engine/Plugins/Experimental/NNERuntimeRDG/Source/NNERuntimeRDG/Private/Dml/NNEDmlOperatorSqueeze.cpp b/Engine/Plugins/Experimental/NNERuntimeRDG/Source/NNERuntimeRDG/Private/Dml/NNEDmlOperatorSqueeze.cpp
--- a/Engine/Plugins/Experimental/NNERuntimeRDG/Source/NNERuntimeRDG/Private/Dml/NNEDmlOperatorSqueeze.cpp
+++ b/Engine/Plugins/Experimental/NNERuntimeRDG/Source/NNERuntimeRDG/Private/Dml/NNEDmlOperatorSqueeze.cpp
@@ -64,7 +64,7 @@ public:
 		return true;
 	}
 
-	virtual int PrepareOutputs(TConstArrayView<NNE::Internal::FTensorRef> InputTensors, TArrayView<NNE::Internal::FTensorRef> OutputTensors) const override
+	virtual int PrepareOutputs(TConstArrayView<NNE::Internal::FTensorRef> InputTensors, TArrayView<NNE::Internal::FTensorRef> OutputTensors) override
 	{
 		TConstArrayView<uint32>	InputShape = InputTensors[0]->GetShape().GetData();
 		
diff --git a/Engine/Plugins/Experimental/NNERuntimeRDG/Source/NNERuntimeRDG/Private/Dml/NNEDmlOperatorTranspose.cpp b/Engine/Plugins/Experimental/NNERuntimeRDG/Source/NNERuntimeRDG/Private/Dml/NNEDmlOperatorTranspose.cpp
--- a/Engine/Plugins/Experimental/NNERuntimeRDG/Source/NNERuntimeRDG/Private/Dml/NNEDmlOperatorTranspose.cpp
+++ b/Engine/Plugins/Experimental/NNERuntimeRDG/Source/NNERuntimeRDG/Private/Dml/NNEDmlOperatorTranspose.cpp
@@ -86,7 +86,7 @@ public:
 
 	}
 
-	virtual int PrepareOutputs(TConstArrayView<NNE::Internal::FTensorRef> InputTensors, TArrayView<NNE::Internal::FTensorRef> OutputTensors) const override
+	virtual int PrepareOutputs(TConstArrayView<NNE::Internal::FTensorRef> InputTensors, TArrayView<NNE::Internal::FTensorRef> OutputTensors) override
 	{
 		Util::FSmallUIntArray OutputShape = Permute(InputTensors[0]->GetShape().GetData());
 
diff --git a/Engine/Plugins/Experimental/NNERuntimeRDG/Source/NNERuntimeRDG/Private/Dml/NNEDmlOperatorUnsqueeze.cpp b/Engine/Plugins/Experimental/NNERuntimeRDG/Source/NNERuntimeRDG/Private/Dml/NNEDmlOperatorUnsqueeze.cpp
--- a/Engine/Plugins/Experimental/NNERuntimeRDG/Source/NNERuntimeRDG/Private/Dml/NNEDmlOperatorUnsqueeze.cpp
+++ b/Engine/Plugins/Experimental/NNERuntimeRDG/Source/NNERuntimeRDG/Private/Dml/NNEDmlOperatorUnsqueeze.cpp
@@ -68,7 +68,7 @@ public:
 		return true;
 	}
 
-	virtual int PrepareOutputs(TConstArrayView<NNE::Internal::FTensorRef> InputTensors, TArrayView<NNE::Internal::FTensorRef> OutputTensors) const override
+	virtual int PrepareOutputs(TConstArrayView<NNE::Internal::FTensorRef> InputTensors, TArrayView<NNE::Internal::FTensorRef> OutputTensors) override
 	{
 		TConstArrayView<uint32>	InputShape = InputTensors[0]->GetShape().GetData();
 		const int32				OutShapeRank = InputShape.Num() + Axes.Num();
diff --git a/Engine/Plugins/Experimental/NNERuntimeRDG/Source/NNERuntimeRDG/Private/Hlsl/NNERuntimeRDGBatchNormalization.cpp b/Engine/Plugins/Experimental/NNERuntimeRDG/Source/NNERuntimeRDG/Private/Hlsl/NNERuntimeRDGBatchNormalization.cpp
--- a/Engine/Plugins/Experimental/NNERuntimeRDG/Source/NNERuntimeRDG/Private/Hlsl/NNERuntimeRDGBatchNormalization.cpp
+++ b/Engine/Plugins/Experimental/NNERuntimeRDG/Source/NNERuntimeRDG/Private/Hlsl/NNERuntimeRDGBatchNormalization.cpp
@@ -25,7 +25,7 @@ namespace UE::NNERuntimeRDG::Private::Hlsl
 
 	public:
 
-		virtual int PrepareOutputs(TConstArrayView<NNE::Internal::FTensorRef> InputTensors, TArrayView<NNE::Internal::FTensorRef> OutputTensors) const override
+		virtual int PrepareOutputs(TConstArrayView<NNE::Internal::FTensorRef> InputTensors, TArrayView<NNE::Internal::FTensorRef> OutputTensors) override
 		{
 			check(InputTensors.Num() == 5);
 			check(OutputTensors.Num() == 1);
diff --git a/Engine/Plugins/Experimental/NNERuntimeRDG/Source/NNERuntimeRDG/Private/Hlsl/NNERuntimeRDGCast.cpp b/Engine/Plugins/Experimental/NNERuntimeRDG/Source/NNERuntimeRDG/Private/Hlsl/NNERuntimeRDGCast.cpp
--- a/Engine/Plugins/Experimental/NNERuntimeRDG/Source/NNERuntimeRDG/Private/Hlsl/NNERuntimeRDGCast.cpp
+++ b/Engine/Plugins/Experimental/NNERuntimeRDG/Source/NNERuntimeRDG/Private/Hlsl/NNERuntimeRDGCast.cpp
@@ -19,7 +19,7 @@ namespace UE::NNERuntimeRDG::Private::Hlsl
 		FCast() {}
 		virtual ~FCast() = default;
 
-		virtual int PrepareOutputs(TConstArrayView<NNE::Internal::FTensorRef> InputTensors, TArrayView<NNE::Internal::FTensorRef> OutputTensors) const override
+		virtual int PrepareOutputs(TConstArrayView<NNE::Internal::FTensorRef> InputTensors, TArrayView<NNE::Internal::FTensorRef> OutputTensors) override
 		{
 			check(InputTensors.Num() == 1);
 			check(OutputTensors.Num() == 1);
diff --git a/Engine/Plugins/Experimental/NNERuntimeRDG/Source/NNERuntimeRDG/Private/Hlsl/NNERuntimeRDGConcat.cpp b/Engine/Plugins/Experimental/NNERuntimeRDG/Source/NNERuntimeRDG/Private/Hlsl/NNERuntimeRDGConcat.cpp
--- a/Engine/Plugins/Experimental/NNERuntimeRDG/Source/NNERuntimeRDG/Private/Hlsl/NNERuntimeRDGConcat.cpp
+++ b/Engine/Plugins/Experimental/NNERuntimeRDG/Source/NNERuntimeRDG/Private/Hlsl/NNERuntimeRDGConcat.cpp
@@ -51,7 +51,7 @@ namespace UE::NNERuntimeRDG::Private::Hlsl
 
 	public:
 
-		virtual int PrepareOutputs(TConstArrayView<NNE::Internal::FTensorRef> InputTensors, TArrayView<NNE::Internal::FTensorRef> OutputTensors) const override
+		virtual int PrepareOutputs(TConstArrayView<NNE::Internal::FTensorRef> InputTensors, TArrayView<NNE::Internal::FTensorRef> OutputTensors) override
 		{
 			check(InputTensors.Num() >= 1);
 			check(OutputTensors.Num() == 1);
diff --git a/Engine/Plugins/Experimental/NNERuntimeRDG/Source/NNERuntimeRDG/Private/Hlsl/NNERuntimeRDGConv.cpp b/Engine/Plugins/Experimental/NNERuntimeRDG/Source/NNERuntimeRDG/Private/Hlsl/NNERuntimeRDGConv.cpp
--- a/Engine/Plugins/Experimental/NNERuntimeRDG/Source/NNERuntimeRDG/Private/Hlsl/NNERuntimeRDGConv.cpp
+++ b/Engine/Plugins/Experimental/NNERuntimeRDG/Source/NNERuntimeRDG/Private/Hlsl/NNERuntimeRDGConv.cpp
@@ -46,7 +46,7 @@ namespace UE::NNERuntimeRDG::Private::Hlsl
 
 	public:
 
-		virtual int PrepareOutputs(TConstArrayView<NNE::Internal::FTensorRef> InputTensors, TArrayView<NNE::Internal::FTensorRef> OutputTensors) const override
+		virtual int PrepareOutputs(TConstArrayView<NNE::Internal::FTensorRef> InputTensors, TArrayView<NNE::Internal::FTensorRef> OutputTensors) override
 		{
 			check(InputTensors.Num() >= 2 && InputTensors.Num() <= 3);
 			check(OutputTensors.Num() == 1);
diff --git a/Engine/Plugins/Experimental/NNERuntimeRDG/Source/NNERuntimeRDG/Private/Hlsl/NNERuntimeRDGConvTranspose.cpp b/Engine/Plugins/Experimental/NNERuntimeRDG/Source/NNERuntimeRDG/Private/Hlsl/NNERuntimeRDGConvTranspose.cpp
--- a/Engine/Plugins/Experimental/NNERuntimeRDG/Source/NNERuntimeRDG/Private/Hlsl/NNERuntimeRDGConvTranspose.cpp
+++ b/Engine/Plugins/Experimental/NNERuntimeRDG/Source/NNERuntimeRDG/Private/Hlsl/NNERuntimeRDGConvTranspose.cpp
@@ -41,7 +41,7 @@ namespace UE::NNERuntimeRDG::Private::Hlsl
 
 	public:
 
-		virtual int PrepareOutputs(TConstArrayView<NNE::Internal::FTensorRef> InputTensors, TArrayView<NNE::Internal::FTensorRef> OutputTensors) const override
+		virtual int PrepareOutputs(TConstArrayView<NNE::Internal::FTensorRef> InputTensors, TArrayView<NNE::Internal::FTensorRef> OutputTensors) override
 		{
 			check(InputTensors.Num() >= 2 && InputTensors.Num() <= 3);
 			check(OutputTensors.Num() == 1);
diff --git a/Engine/Plugins/Experimental/NNERuntimeRDG/Source/NNERuntimeRDG/Private/Hlsl/NNERuntimeRDGDropout.cpp b/Engine/Plugins/Experimental/NNERuntimeRDG/Source/NNERuntimeRDG/Private/Hlsl/NNERuntimeRDGDropout.cpp
--- a/Engine/Plugins/Experimental/NNERuntimeRDG/Source/NNERuntimeRDG/Private/Hlsl/NNERuntimeRDGDropout.cpp
+++ b/Engine/Plugins/Experimental/NNERuntimeRDG/Source/NNERuntimeRDG/Private/Hlsl/NNERuntimeRDGDropout.cpp
@@ -21,7 +21,7 @@ namespace UE::NNERuntimeRDG::Private::Hlsl
 
 	public:
 
-		virtual int PrepareOutputs(TConstArrayView<NNE::Internal::FTensorRef> InputTensors, TArrayView<NNE::Internal::FTensorRef> OutputTensors) const override
+		virtual int PrepareOutputs(TConstArrayView<NNE::Internal::FTensorRef> InputTensors, TArrayView<NNE::Internal::FTensorRef> OutputTensors) override
 		{
 			check(InputTensors.Num() == 1);
 			check(OutputTensors.Num() == 1);
diff --git a/Engine/Plugins/Experimental/NNERuntimeRDG/Source/NNERuntimeRDG/Private/Hlsl/NNERuntimeRDGElementWiseBinary.cpp b/Engine/Plugins/Experimental/NNERuntimeRDG/Source/NNERuntimeRDG/Private/Hlsl/NNERuntimeRDGElementWiseBinary.cpp
--- a/Engine/Plugins/Experimental/NNERuntimeRDG/Source/NNERuntimeRDG/Private/Hlsl/NNERuntimeRDGElementWiseBinary.cpp
+++ b/Engine/Plugins/Experimental/NNERuntimeRDG/Source/NNERuntimeRDG/Private/Hlsl/NNERuntimeRDGElementWiseBinary.cpp
@@ -27,7 +27,7 @@ namespace UE::NNERuntimeRDG::Private::Hlsl
 
 	public:
 
-		virtual int PrepareOutputs(TConstArrayView<NNE::Internal::FTensorRef> InputTensors, TArrayView<NNE::Internal::FTensorRef> OutputTensors) const override
+		virtual int PrepareOutputs(TConstArrayView<NNE::Internal::FTensorRef> InputTensors, TArrayView<NNE::Internal::FTensorRef> OutputTensors) override
 		{
 			check(InputTensors.Num() == 2);
 			check(OutputTensors.Num() == 1);
diff --git a/Engine/Plugins/Experimental/NNERuntimeRDG/Source/NNERuntimeRDG/Private/Hlsl/NNERuntimeRDGElementWiseUnary.cpp b/Engine/Plugins/Experimental/NNERuntimeRDG/Source/NNERuntimeRDG/Private/Hlsl/NNERuntimeRDGElementWiseUnary.cpp
--- a/Engine/Plugins/Experimental/NNERuntimeRDG/Source/NNERuntimeRDG/Private/Hlsl/NNERuntimeRDGElementWiseUnary.cpp
+++ b/Engine/Plugins/Experimental/NNERuntimeRDG/Source/NNERuntimeRDG/Private/Hlsl/NNERuntimeRDGElementWiseUnary.cpp
@@ -34,7 +34,7 @@ namespace UE::NNERuntimeRDG::Private::Hlsl
 
 	public:
 
-		virtual int PrepareOutputs(TConstArrayView<NNE::Internal::FTensorRef> InputTensors, TArrayView<NNE::Internal::FTensorRef> OutputTensors) const override
+		virtual int PrepareOutputs(TConstArrayView<NNE::Internal::FTensorRef> InputTensors, TArrayView<NNE::Internal::FTensorRef> OutputTensors) override
 		{
 			check(InputTensors.Num() == 1);
 			check(OutputTensors.Num() == 1);
diff --git a/Engine/Plugins/Experimental/NNERuntimeRDG/Source/NNERuntimeRDG/Private/Hlsl/NNERuntimeRDGElementWiseVariadic.cpp b/Engine/Plugins/Experimental/NNERuntimeRDG/Source/NNERuntimeRDG/Private/Hlsl/NNERuntimeRDGElementWiseVariadic.cpp
--- a/Engine/Plugins/Experimental/NNERuntimeRDG/Source/NNERuntimeRDG/Private/Hlsl/NNERuntimeRDGElementWiseVariadic.cpp
+++ b/Engine/Plugins/Experimental/NNERuntimeRDG/Source/NNERuntimeRDG/Private/Hlsl/NNERuntimeRDGElementWiseVariadic.cpp
@@ -94,7 +94,7 @@ namespace UE::NNERuntimeRDG::Private::Hlsl
 
 	public:
 
-		virtual int PrepareOutputs(TConstArrayView<NNE::Internal::FTensorRef> InputTensors, TArrayView<NNE::Internal::FTensorRef> OutputTensors) const override
+		virtual int PrepareOutputs(TConstArrayView<NNE::Internal::FTensorRef> InputTensors, TArrayView<NNE::Internal::FTensorRef> OutputTensors) override
 		{
 			check(InputTensors.Num() > 0);
 			check(OutputTensors.Num() == 1);
diff --git a/Engine/Plugins/Experimental/NNERuntimeRDG/Source/NNERuntimeRDG/Private/Hlsl/NNERuntimeRDGFlatten.cpp b/Engine/Plugins/Experimental/NNERuntimeRDG/Source/NNERuntimeRDG/Private/Hlsl/NNERuntimeRDGFlatten.cpp
--- a/Engine/Plugins/Experimental/NNERuntimeRDG/Source/NNERuntimeRDG/Private/Hlsl/NNERuntimeRDGFlatten.cpp
+++ b/Engine/Plugins/Experimental/NNERuntimeRDG/Source/NNERuntimeRDG/Private/Hlsl/NNERuntimeRDGFlatten.cpp
@@ -23,7 +23,7 @@ namespace UE::NNERuntimeRDG::Private::Hlsl
 
 	public:
 
-		virtual int PrepareOutputs(TConstArrayView<NNE::Internal::FTensorRef> InputTensors, TArrayView<NNE::Internal::FTensorRef> OutputTensors) const override
+		virtual int PrepareOutputs(TConstArrayView<NNE::Internal::FTensorRef> InputTensors, TArrayView<NNE::Internal::FTensorRef> OutputTensors) override
 		{
 			check(InputTensors.Num() == 1);
 			check(OutputTensors.Num() == 1);
diff --git a/Engine/Plugins/Experimental/NNERuntimeRDG/Source/NNERuntimeRDG/Private/Hlsl/NNERuntimeRDGGather.cpp b/Engine/Plugins/Experimental/NNERuntimeRDG/Source/NNERuntimeRDG/Private/Hlsl/NNERuntimeRDGGather.cpp
--- a/Engine/Plugins/Experimental/NNERuntimeRDG/Source/NNERuntimeRDG/Private/Hlsl/NNERuntimeRDGGather.cpp
+++ b/Engine/Plugins/Experimental/NNERuntimeRDG/Source/NNERuntimeRDG/Private/Hlsl/NNERuntimeRDGGather.cpp
@@ -29,7 +29,7 @@ namespace UE::NNERuntimeRDG::Private::Hlsl
 
 	public:
 
-		virtual int PrepareOutputs(TConstArrayView<NNE::Internal::FTensorRef> InputTensors, TArrayView<NNE::Internal::FTensorRef> OutputTensors) const override
+		virtual int PrepareOutputs(TConstArrayView<NNE::Internal::FTensorRef> InputTensors, TArrayView<NNE::Internal::FTensorRef> OutputTensors) override
 		{
 			check(InputTensors.Num() == 2)
 			check(OutputTensors.Num() == 1)
diff --git a/Engine/Plugins/Experimental/NNERuntimeRDG/Source/NNERuntimeRDG/Private/Hlsl/NNERuntimeRDGGemm.cpp b/Engine/Plugins/Experimental/NNERuntimeRDG/Source/NNERuntimeRDG/Private/Hlsl/NNERuntimeRDGGemm.cpp
--- a/Engine/Plugins/Experimental/NNERuntimeRDG/Source/NNERuntimeRDG/Private/Hlsl/NNERuntimeRDGGemm.cpp
+++ b/Engine/Plugins/Experimental/NNERuntimeRDG/Source/NNERuntimeRDG/Private/Hlsl/NNERuntimeRDGGemm.cpp
@@ -30,7 +30,7 @@ namespace UE::NNERuntimeRDG::Private::Hlsl
 
 	public:
 
-		virtual int PrepareOutputs(TConstArrayView<NNE::Internal::FTensorRef> InputTensors, TArrayView<NNE::Internal::FTensorRef> OutputTensors) const override
+		virtual int PrepareOutputs(TConstArrayView<NNE::Internal::FTensorRef> InputTensors, TArrayView<NNE::Internal::FTensorRef> OutputTensors) override
 		{
 			check(InputTensors.Num() >= 2 && InputTensors.Num() <= 3);
 			check(OutputTensors.Num() == 1);
diff --git a/Engine/Plugins/Experimental/NNERuntimeRDG/Source/NNERuntimeRDG/Private/Hlsl/NNERuntimeRDGGlobalPool.cpp b/Engine/Plugins/Experimental/NNERuntimeRDG/Source/NNERuntimeRDG/Private/Hlsl/NNERuntimeRDGGlobalPool.cpp
--- a/Engine/Plugins/Experimental/NNERuntimeRDG/Source/NNERuntimeRDG/Private/Hlsl/NNERuntimeRDGGlobalPool.cpp
+++ b/Engine/Plugins/Experimental/NNERuntimeRDG/Source/NNERuntimeRDG/Private/Hlsl/NNERuntimeRDGGlobalPool.cpp
@@ -29,7 +29,7 @@ namespace UE::NNERuntimeRDG::Private::Hlsl
 
 	public:
 
-		virtual int PrepareOutputs(TConstArrayView<NNE::Internal::FTensorRef> InputTensors, TArrayView<NNE::Internal::FTensorRef> OutputTensors) const override
+		virtual int PrepareOutputs(TConstArrayView<NNE::Internal::FTensorRef> InputTensors, TArrayView<NNE::Internal::FTensorRef> OutputTensors) override
 		{
 			check(InputTensors.Num() == 1);
 			check(OutputTensors.Num() == 1);
diff --git a/Engine/Plugins/Experimental/NNERuntimeRDG/Source/NNERuntimeRDG/Private/Hlsl/NNERuntimeRDGIdentity.cpp b/Engine/Plugins/Experimental/NNERuntimeRDG/Source/NNERuntimeRDG/Private/Hlsl/NNERuntimeRDGIdentity.cpp
--- a/Engine/Plugins/Experimental/NNERuntimeRDG/Source/NNERuntimeRDG/Private/Hlsl/NNERuntimeRDGIdentity.cpp
+++ b/Engine/Plugins/Experimental/NNERuntimeRDG/Source/NNERuntimeRDG/Private/Hlsl/NNERuntimeRDGIdentity.cpp
@@ -21,7 +21,7 @@ namespace UE::NNERuntimeRDG::Private::Hlsl
 
 	public:
 
-		virtual int PrepareOutputs(TConstArrayView<NNE::Internal::FTensorRef> InputTensors, TArrayView<NNE::Internal::FTensorRef> OutputTensors) const override
+		virtual int PrepareOutputs(TConstArrayView<NNE::Internal::FTensorRef> InputTensors, TArrayView<NNE::Internal::FTensorRef> OutputTensors) override
 		{
 			check(InputTensors.Num() == 1);
 			check(OutputTensors.Num() == 1);
diff --git a/Engine/Plugins/Experimental/NNERuntimeRDG/Source/NNERuntimeRDG/Private/Hlsl/NNERuntimeRDGInstanceNormalization.cpp b/Engine/Plugins/Experimental/NNERuntimeRDG/Source/NNERuntimeRDG/Private/Hlsl/NNERuntimeRDGInstanceNormalization.cpp
--- a/Engine/Plugins/Experimental/NNERuntimeRDG/Source/NNERuntimeRDG/Private/Hlsl/NNERuntimeRDGInstanceNormalization.cpp
+++ b/Engine/Plugins/Experimental/NNERuntimeRDG/Source/NNERuntimeRDG/Private/Hlsl/NNERuntimeRDGInstanceNormalization.cpp
@@ -59,7 +59,7 @@ namespace UE::NNERuntimeRDG::Private::Hlsl
 
 	public:
 
-		virtual int PrepareOutputs(TConstArrayView<NNE::Internal::FTensorRef> InputTensors, TArrayView<NNE::Internal::FTensorRef> OutputTensors) const override
+		virtual int PrepareOutputs(TConstArrayView<NNE::Internal::FTensorRef> InputTensors, TArrayView<NNE::Internal::FTensorRef> OutputTensors) override
 		{
 			check(InputTensors.Num() == 3);
 			check(OutputTensors.Num() == 1);
diff --git a/Engine/Plugins/Experimental/NNERuntimeRDG/Source/NNERuntimeRDG/Private/Hlsl/NNERuntimeRDGMatMul.cpp b/Engine/Plugins/Experimental/NNERuntimeRDG/Source/NNERuntimeRDG/Private/Hlsl/NNERuntimeRDGMatMul.cpp
--- a/Engine/Plugins/Experimental/NNERuntimeRDG/Source/NNERuntimeRDG/Private/Hlsl/NNERuntimeRDGMatMul.cpp
+++ b/Engine/Plugins/Experimental/NNERuntimeRDG/Source/NNERuntimeRDG/Private/Hlsl/NNERuntimeRDGMatMul.cpp
@@ -22,7 +22,7 @@ namespace UE::NNERuntimeRDG::Private::Hlsl
 
 	public:
 
-		virtual int PrepareOutputs(TConstArrayView<NNE::Internal::FTensorRef> InputTensors, TArrayView<NNE::Internal::FTensorRef> OutputTensors) const override
+		virtual int PrepareOutputs(TConstArrayView<NNE::Internal::FTensorRef> InputTensors, TArrayView<NNE::Internal::FTensorRef> OutputTensors) override
 		{
 			check(InputTensors.Num() == 2);
 			check(OutputTensors.Num() == 1);
diff --git a/Engine/Plugins/Experimental/NNERuntimeRDG/Source/NNERuntimeRDG/Private/Hlsl/NNERuntimeRDGPad.cpp b/Engine/Plugins/Experimental/NNERuntimeRDG/Source/NNERuntimeRDG/Private/Hlsl/NNERuntimeRDGPad.cpp
--- a/Engine/Plugins/Experimental/NNERuntimeRDG/Source/NNERuntimeRDG/Private/Hlsl/NNERuntimeRDGPad.cpp
+++ b/Engine/Plugins/Experimental/NNERuntimeRDG/Source/NNERuntimeRDG/Private/Hlsl/NNERuntimeRDGPad.cpp
@@ -28,7 +28,7 @@ namespace UE::NNERuntimeRDG::Private::Hlsl
 
 	public:
 
-		virtual int PrepareOutputs(TConstArrayView<NNE::Internal::FTensorRef> InputTensors, TArrayView<NNE::Internal::FTensorRef> OutputTensors) const override
+		virtual int PrepareOutputs(TConstArrayView<NNE::Internal::FTensorRef> InputTensors, TArrayView<NNE::Internal::FTensorRef> OutputTensors) override
 		{
 			check(InputTensors.Num() == 1);
 			check(OutputTensors.Num() == 1);
diff --git a/Engine/Plugins/Experimental/NNERuntimeRDG/Source/NNERuntimeRDG/Private/Hlsl/NNERuntimeRDGPool.cpp b/Engine/Plugins/Experimental/NNERuntimeRDG/Source/NNERuntimeRDG/Private/Hlsl/NNERuntimeRDGPool.cpp
--- a/Engine/Plugins/Experimental/NNERuntimeRDG/Source/NNERuntimeRDG/Private/Hlsl/NNERuntimeRDGPool.cpp
+++ b/Engine/Plugins/Experimental/NNERuntimeRDG/Source/NNERuntimeRDG/Private/Hlsl/NNERuntimeRDGPool.cpp
@@ -34,7 +34,7 @@ namespace UE::NNERuntimeRDG::Private::Hlsl
 
 	public:
 
-		virtual int PrepareOutputs(TConstArrayView<NNE::Internal::FTensorRef> InputTensors, TArrayView<NNE::Internal::FTensorRef> OutputTensors) const override
+		virtual int PrepareOutputs(TConstArrayView<NNE::Internal::FTensorRef> InputTensors, TArrayView<NNE::Internal::FTensorRef> OutputTensors) override
 		{
 			check(InputTensors.Num() == 1);
 			check(OutputTensors.Num() == 1);
diff --git a/Engine/Plugins/Experimental/NNERuntimeRDG/Source/NNERuntimeRDG/Private/Hlsl/NNERuntimeRDGReduce.cpp b/Engine/Plugins/Experimental/NNERuntimeRDG/Source/NNERuntimeRDG/Private/Hlsl/NNERuntimeRDGReduce.cpp
--- a/Engine/Plugins/Experimental/NNERuntimeRDG/Source/NNERuntimeRDG/Private/Hlsl/NNERuntimeRDGReduce.cpp
+++ b/Engine/Plugins/Experimental/NNERuntimeRDG/Source/NNERuntimeRDG/Private/Hlsl/NNERuntimeRDGReduce.cpp
@@ -30,7 +30,7 @@ namespace UE::NNERuntimeRDG::Private::Hlsl
 
 	public:
 
-		virtual int PrepareOutputs(TConstArrayView<NNE::Internal::FTensorRef> InputTensors, TArrayView<NNE::Internal::FTensorRef> OutputTensors) const override
+		virtual int PrepareOutputs(TConstArrayView<NNE::Internal::FTensorRef> InputTensors, TArrayView<NNE::Internal::FTensorRef> OutputTensors) override
 		{
 			check(InputTensors.Num() == 1);
 			check(OutputTensors.Num() == 1);
diff --git a/Engine/Plugins/Experimental/NNERuntimeRDG/Source/NNERuntimeRDG/Private/Hlsl/NNERuntimeRDGReshape.cpp b/Engine/Plugins/Experimental/NNERuntimeRDG/Source/NNERuntimeRDG/Private/Hlsl/NNERuntimeRDGReshape.cpp
--- a/Engine/Plugins/Experimental/NNERuntimeRDG/Source/NNERuntimeRDG/Private/Hlsl/NNERuntimeRDGReshape.cpp
+++ b/Engine/Plugins/Experimental/NNERuntimeRDG/Source/NNERuntimeRDG/Private/Hlsl/NNERuntimeRDGReshape.cpp
@@ -24,7 +24,7 @@ namespace UE::NNERuntimeRDG::Private::Hlsl
 
 	public:
 
-		virtual int PrepareOutputs(TConstArrayView<NNE::Internal::FTensorRef> InputTensors, TArrayView<NNE::Internal::FTensorRef> OutputTensors) const override
+		virtual int PrepareOutputs(TConstArrayView<NNE::Internal::FTensorRef> InputTensors, TArrayView<NNE::Internal::FTensorRef> OutputTensors) override
 		{
 			check(InputTensors.Num() == 2);
 			check(OutputTensors.Num() == 1);
diff --git a/Engine/Plugins/Experimental/NNERuntimeRDG/Source/NNERuntimeRDG/Private/Hlsl/NNERuntimeRDGShape.cpp b/Engine/Plugins/Experimental/NNERuntimeRDG/Source/NNERuntimeRDG/Private/Hlsl/NNERuntimeRDGShape.cpp
--- a/Engine/Plugins/Experimental/NNERuntimeRDG/Source/NNERuntimeRDG/Private/Hlsl/NNERuntimeRDGShape.cpp
+++ b/Engine/Plugins/Experimental/NNERuntimeRDG/Source/NNERuntimeRDG/Private/Hlsl/NNERuntimeRDGShape.cpp
@@ -18,7 +18,7 @@ namespace UE::NNERuntimeRDG::Private::Hlsl
 
 	public:
 
-		virtual int PrepareOutputs(TConstArrayView<NNE::Internal::FTensorRef> InputTensors, TArrayView<NNE::Internal::FTensorRef> OutputTensors) const override
+		virtual int PrepareOutputs(TConstArrayView<NNE::Internal::FTensorRef> InputTensors, TArrayView<NNE::Internal::FTensorRef> OutputTensors) override
 		{
 			check(InputTensors.Num() == 1);
 			check(OutputTensors.Num() == 1);
diff --git a/Engine/Plugins/Experimental/NNERuntimeRDG/Source/NNERuntimeRDG/Private/Hlsl/NNERuntimeRDGSize.cpp b/Engine/Plugins/Experimental/NNERuntimeRDG/Source/NNERuntimeRDG/Private/Hlsl/NNERuntimeRDGSize.cpp
--- a/Engine/Plugins/Experimental/NNERuntimeRDG/Source/NNERuntimeRDG/Private/Hlsl/NNERuntimeRDGSize.cpp
+++ b/Engine/Plugins/Experimental/NNERuntimeRDG/Source/NNERuntimeRDG/Private/Hlsl/NNERuntimeRDGSize.cpp
@@ -20,7 +20,7 @@ namespace UE::NNERuntimeRDG::Private::Hlsl
 
 	public:
 
-		virtual int PrepareOutputs(TConstArrayView<NNE::Internal::FTensorRef> InputTensors, TArrayView<NNE::Internal::FTensorRef> OutputTensors) const override
+		virtual int PrepareOutputs(TConstArrayView<NNE::Internal::FTensorRef> InputTensors, TArrayView<NNE::Internal::FTensorRef> OutputTensors) override
 		{
 			check(InputTensors.Num() == 1);
 			check(OutputTensors.Num() == 1);
diff --git a/Engine/Plugins/Experimental/NNERuntimeRDG/Source/NNERuntimeRDG/Private/Hlsl/NNERuntimeRDGSlice.cpp b/Engine/Plugins/Experimental/NNERuntimeRDG/Source/NNERuntimeRDG/Private/Hlsl/NNERuntimeRDGSlice.cpp
--- a/Engine/Plugins/Experimental/NNERuntimeRDG/Source/NNERuntimeRDG/Private/Hlsl/NNERuntimeRDGSlice.cpp
+++ b/Engine/Plugins/Experimental/NNERuntimeRDG/Source/NNERuntimeRDG/Private/Hlsl/NNERuntimeRDGSlice.cpp
@@ -127,7 +127,7 @@ namespace UE::NNERuntimeRDG::Private::Hlsl
 
 	public:
 
-		virtual int PrepareOutputs(TConstArrayView<NNE::Internal::FTensorRef> InputTensors, TArrayView<NNE::Internal::FTensorRef> OutputTensors) const override
+		virtual int PrepareOutputs(TConstArrayView<NNE::Internal::FTensorRef> InputTensors, TArrayView<NNE::Internal::FTensorRef> OutputTensors) override
 		{
 			check(OutputTensors.Num() == 1);
 
diff --git a/Engine/Plugins/Experimental/NNERuntimeRDG/Source/NNERuntimeRDG/Private/Hlsl/NNERuntimeRDGSoftmax.cpp b/Engine/Plugins/Experimental/NNERuntimeRDG/Source/NNERuntimeRDG/Private/Hlsl/NNERuntimeRDGSoftmax.cpp
--- a/Engine/Plugins/Experimental/NNERuntimeRDG/Source/NNERuntimeRDG/Private/Hlsl/NNERuntimeRDGSoftmax.cpp
+++ b/Engine/Plugins/Experimental/NNERuntimeRDG/Source/NNERuntimeRDG/Private/Hlsl/NNERuntimeRDGSoftmax.cpp
@@ -15,6 +15,7 @@ namespace UE::NNERuntimeRDG::Private::Hlsl
 	/**
 	 * Softmax operator implementation
 	 */
+	template< NNEHlslShaders::Internal::ESoftmaxOperatorType SoftmaxOperatorType >
 	class FSoftmax : public FOperatorHlsl
 	{
 
@@ -29,7 +30,7 @@ namespace UE::NNERuntimeRDG::Private::Hlsl
 
 	public:
 
-		virtual int PrepareOutputs(TConstArrayView<NNE::Internal::FTensorRef> InputTensors, TArrayView<NNE::Internal::FTensorRef> OutputTensors) const override
+		virtual int PrepareOutputs(TConstArrayView<NNE::Internal::FTensorRef> InputTensors, TArrayView<NNE::Internal::FTensorRef> OutputTensors) override
 		{
 			check(InputTensors.Num() == 1);
 			check(OutputTensors.Num() == 1);
@@ -113,7 +114,9 @@ namespace UE::NNERuntimeRDG::Private::Hlsl
 			SoftmaxParameters->Input = GraphBuilder.CreateSRV(FRDGBufferSRVDesc(Input.GetBuffer(), PF_R32_FLOAT));
 			SoftmaxParameters->InputSumExp = GraphBuilder.CreateSRV(FRDGBufferSRVDesc(SumExpBuffer, PF_R32_FLOAT));
 			SoftmaxParameters->Output = GraphBuilder.CreateUAV(FRDGBufferUAVDesc(Output.GetBuffer(), PF_R32_FLOAT));
-			TShaderMapRef<TSoftmaxCS> ComputeShader(GetGlobalShaderMap(GMaxRHIFeatureLevel));
+			TSoftmaxCS::FPermutationDomain PermutationVector;
+			PermutationVector.Set<TSoftmaxCS::FSoftmaxType>(SoftmaxOperatorType);
+			TShaderMapRef<TSoftmaxCS> ComputeShader(GetGlobalShaderMap(GMaxRHIFeatureLevel), PermutationVector);
 
 			FComputeShaderUtils::AddPass(
 				GraphBuilder,
@@ -145,12 +148,19 @@ namespace UE::NNERuntimeRDG::Private::Hlsl
 
 	FOperatorHlsl* CreateSoftmaxOperator()
 	{
-		return new FSoftmax();
+		return new FSoftmax<NNEHlslShaders::Internal::ESoftmaxOperatorType::SOFTMAX>();
+	}
+
+	FOperatorHlsl* CreateLogSoftmaxOperator()
+	{
+		return new FSoftmax<NNEHlslShaders::Internal::ESoftmaxOperatorType::LOG_SOFTMAX>();
 	}
 
 	bool RegisterSoftmaxOperator(FOperatorRegistryHlsl& Registry)
 	{
-		Registry.OpAdd({{TEXT("Softmax"), TEXT("Onnx")}}, CreateSoftmaxOperator, ValidateSoftmaxOperator);
+		Registry.OpAdd({ {TEXT("Softmax"), TEXT("Onnx")}, 1 }, CreateSoftmaxOperator, ValidateSoftmaxOperator);
+		Registry.OpAdd({ {TEXT("LogSoftmax"), TEXT("Onnx")}, 1 }, CreateLogSoftmaxOperator, ValidateSoftmaxOperator);
+		
 		return true;
 	}
 } // UE::NNERuntimeRDG::Private::Hlsl
diff --git a/Engine/Plugins/Experimental/NNERuntimeRDG/Source/NNERuntimeRDG/Private/Hlsl/NNERuntimeRDGSqueeze.cpp b/Engine/Plugins/Experimental/NNERuntimeRDG/Source/NNERuntimeRDG/Private/Hlsl/NNERuntimeRDGSqueeze.cpp
--- a/Engine/Plugins/Experimental/NNERuntimeRDG/Source/NNERuntimeRDG/Private/Hlsl/NNERuntimeRDGSqueeze.cpp
+++ b/Engine/Plugins/Experimental/NNERuntimeRDG/Source/NNERuntimeRDG/Private/Hlsl/NNERuntimeRDGSqueeze.cpp
@@ -23,7 +23,7 @@ namespace UE::NNERuntimeRDG::Private::Hlsl
 
 	public:
 
-		virtual int PrepareOutputs(TConstArrayView<NNE::Internal::FTensorRef> InputTensors, TArrayView<NNE::Internal::FTensorRef> OutputTensors) const override
+		virtual int PrepareOutputs(TConstArrayView<NNE::Internal::FTensorRef> InputTensors, TArrayView<NNE::Internal::FTensorRef> OutputTensors) override
 		{
 			check(InputTensors.Num() == 1);
 			check(OutputTensors.Num() == 1);
diff --git a/Engine/Plugins/Experimental/NNERuntimeRDG/Source/NNERuntimeRDG/Private/Hlsl/NNERuntimeRDGTranspose.cpp b/Engine/Plugins/Experimental/NNERuntimeRDG/Source/NNERuntimeRDG/Private/Hlsl/NNERuntimeRDGTranspose.cpp
--- a/Engine/Plugins/Experimental/NNERuntimeRDG/Source/NNERuntimeRDG/Private/Hlsl/NNERuntimeRDGTranspose.cpp
+++ b/Engine/Plugins/Experimental/NNERuntimeRDG/Source/NNERuntimeRDG/Private/Hlsl/NNERuntimeRDGTranspose.cpp
@@ -27,7 +27,7 @@ namespace UE::NNERuntimeRDG::Private::Hlsl
 
 	public:
 
-		virtual int PrepareOutputs(TConstArrayView<NNE::Internal::FTensorRef> InputTensors, TArrayView<NNE::Internal::FTensorRef> OutputTensors) const override
+		virtual int PrepareOutputs(TConstArrayView<NNE::Internal::FTensorRef> InputTensors, TArrayView<NNE::Internal::FTensorRef> OutputTensors) override
 		{
 			check(InputTensors.Num() == 1);
 			check(OutputTensors.Num() == 1);
diff --git a/Engine/Plugins/Experimental/NNERuntimeRDG/Source/NNERuntimeRDG/Private/Hlsl/NNERuntimeRDGUnsqueeze.cpp b/Engine/Plugins/Experimental/NNERuntimeRDG/Source/NNERuntimeRDG/Private/Hlsl/NNERuntimeRDGUnsqueeze.cpp
--- a/Engine/Plugins/Experimental/NNERuntimeRDG/Source/NNERuntimeRDG/Private/Hlsl/NNERuntimeRDGUnsqueeze.cpp
+++ b/Engine/Plugins/Experimental/NNERuntimeRDG/Source/NNERuntimeRDG/Private/Hlsl/NNERuntimeRDGUnsqueeze.cpp
@@ -23,7 +23,7 @@ namespace UE::NNERuntimeRDG::Private::Hlsl
 
 	public:
 
-		virtual int PrepareOutputs(TConstArrayView<NNE::Internal::FTensorRef> InputTensors, TArrayView<NNE::Internal::FTensorRef> OutputTensors) const override
+		virtual int PrepareOutputs(TConstArrayView<NNE::Internal::FTensorRef> InputTensors, TArrayView<NNE::Internal::FTensorRef> OutputTensors) override
 		{
 			check(InputTensors.Num() == 1);
 			check(OutputTensors.Num() == 1);
diff --git a/Engine/Plugins/Experimental/NNERuntimeRDG/Source/NNERuntimeRDG/Private/Hlsl/NNERuntimeRDGUpsample.cpp b/Engine/Plugins/Experimental/NNERuntimeRDG/Source/NNERuntimeRDG/Private/Hlsl/NNERuntimeRDGUpsample.cpp
--- a/Engine/Plugins/Experimental/NNERuntimeRDG/Source/NNERuntimeRDG/Private/Hlsl/NNERuntimeRDGUpsample.cpp
+++ b/Engine/Plugins/Experimental/NNERuntimeRDG/Source/NNERuntimeRDG/Private/Hlsl/NNERuntimeRDGUpsample.cpp
@@ -27,7 +27,7 @@ namespace UE::NNERuntimeRDG::Private::Hlsl
 
 	public:
 
-		virtual int PrepareOutputs(TConstArrayView<NNE::Internal::FTensorRef> InputTensors, TArrayView<NNE::Internal::FTensorRef> OutputTensors) const override
+		virtual int PrepareOutputs(TConstArrayView<NNE::Internal::FTensorRef> InputTensors, TArrayView<NNE::Internal::FTensorRef> OutputTensors) override
 		{
 			check(InputTensors.Num() == 2);
 			check(OutputTensors.Num() == 1);
diff --git a/Engine/Plugins/Experimental/NNERuntimeRDG/Source/NNERuntimeRDG/Private/NNERuntimeRDGBase.h b/Engine/Plugins/Experimental/NNERuntimeRDG/Source/NNERuntimeRDG/Private/NNERuntimeRDGBase.h
--- a/Engine/Plugins/Experimental/NNERuntimeRDG/Source/NNERuntimeRDG/Private/NNERuntimeRDGBase.h
+++ b/Engine/Plugins/Experimental/NNERuntimeRDG/Source/NNERuntimeRDG/Private/NNERuntimeRDGBase.h
@@ -33,7 +33,7 @@ namespace UE::NNERuntimeRDG::Private
 struct IPrepareOperator
 {
 	virtual ~IPrepareOperator() = default;
-	virtual int PrepareOutputs(TConstArrayView<NNE::Internal::FTensorRef> InputTensors, TArrayView<NNE::Internal::FTensorRef> OutputTensors) const = 0;
+	virtual int PrepareOutputs(TConstArrayView<NNE::Internal::FTensorRef> InputTensors, TArrayView<NNE::Internal::FTensorRef> OutputTensors) = 0;
 };
 
 /**
diff --git a/Engine/Plugins/Experimental/NNERuntimeRDG/Source/NNERuntimeRDG/Private/NNERuntimeRDGModelHlsl.cpp b/Engine/Plugins/Experimental/NNERuntimeRDG/Source/NNERuntimeRDG/Private/NNERuntimeRDGModelHlsl.cpp
--- a/Engine/Plugins/Experimental/NNERuntimeRDG/Source/NNERuntimeRDG/Private/NNERuntimeRDGModelHlsl.cpp
+++ b/Engine/Plugins/Experimental/NNERuntimeRDG/Source/NNERuntimeRDG/Private/NNERuntimeRDGModelHlsl.cpp
@@ -261,7 +261,7 @@ int FModelInstance::PrepareTensorShapesAndData()
 			checkCode(AllInitializedTensors[i] = true);
 		}
 
-		const FOperatorHlsl* Op = Operators[Idx];
+		FOperatorHlsl* Op = Operators[Idx];
 
 		if (Op->PrepareOutputs(InputTensors, OutputTensors) != 0)
 		{
