[NNE] Prevent memory leak
Problem
Instances of the FOperatorDml are not deleted in destructor

Solution
Delete instances of the FOperatorDml are not deleted in destructor

Tests
- NNEIncubatorEditor Win64
- NNEIncubator standalone Win64
- Non-unity compile

#rnx

diff --git a/Engine/Plugins/Experimental/NNERuntimeRDG/Source/NNERuntimeRDG/Private/Dml/NNEDmlModel.cpp b/Engine/Plugins/Experimental/NNERuntimeRDG/Source/NNERuntimeRDG/Private/Dml/NNEDmlModel.cpp
--- a/Engine/Plugins/Experimental/NNERuntimeRDG/Source/NNERuntimeRDG/Private/Dml/NNEDmlModel.cpp
+++ b/Engine/Plugins/Experimental/NNERuntimeRDG/Source/NNERuntimeRDG/Private/Dml/NNEDmlModel.cpp
@@ -870,12 +870,19 @@ FModelInstance::~FModelInstance()
 
 	Signal->Wait();
 	FGenericPlatformProcess::ReturnSynchEventToPool(Signal);
+
+	for (FOperatorDml* Op : Operators)
+	{
+		delete Op;
+	}
+	Operators.Reset();
 }
 
 bool FModelInstance::Init(TConstArrayView<uint8> ModelData, FDmlDeviceContext* InDevCtx)
 {
+	check(Operators.IsEmpty());
 	ConstantCPUTensorIndices.Reset();
-
+	
 	check(ModelData.Num() > 0);
 	FNNERuntimeFormat	Format;
 	int32 GuidSize = FModelInfo::Get()->GetGuidSize();
