Remove unneeded functions causing to silence warning V575: The null pointer is passed into 'TConstArrayView' function. Inspect the first argument.

diff --git a/Engine/Plugins/Experimental/NNERuntimeRDG/Source/NNERuntimeRDG/Private/Dml/NNEDmlOperator.h b/Engine/Plugins/Experimental/NNERuntimeRDG/Source/NNERuntimeRDG/Private/Dml/NNEDmlOperator.h
--- a/Engine/Plugins/Experimental/NNERuntimeRDG/Source/NNERuntimeRDG/Private/Dml/NNEDmlOperator.h
+++ b/Engine/Plugins/Experimental/NNERuntimeRDG/Source/NNERuntimeRDG/Private/Dml/NNEDmlOperator.h
@@ -55,18 +55,6 @@ static constexpr uint32 NcdhwDimensionCount = 5;
 static constexpr uint32 NcdhwSpatialDimensionCount = 3;
 static constexpr uint32 NonspatialDimensionCount = 2;		//!< The batch and channel dimensions of NCW, NCHW, NCDHW....
 
-template<typename T>
-inline TArrayView<T> MakeEmptyArrayView()
-{
-	return MakeArrayView(static_cast<T*>(nullptr), 0);
-}
-
-template<typename T>
-inline TConstArrayView<T> MakeEmptyConstArrayView()
-{
-	return TConstArrayView<T>(static_cast<const T*>(nullptr), 0);
-}
-
 class FDmlDeviceContext
 {
 public:
@@ -124,7 +112,7 @@ template<typename OutputType, typename AllocatorType>
 inline bool GetArrayAttributeNoOverflow(
 	const FNNEAttributeValue* Attr, 
 	TArray<OutputType, AllocatorType>& OutputArray, 
-	TConstArrayView<OutputType> DefaultValues = MakeEmptyConstArrayView<OutputType>()
+	TConstArrayView<OutputType> DefaultValues = TConstArrayView<OutputType>()
 	)
 {
 	if (Attr)
diff --git a/Engine/Plugins/Experimental/NNERuntimeRDG/Source/NNERuntimeRDG/Private/Dml/NNEDmlOperatorUtils.h b/Engine/Plugins/Experimental/NNERuntimeRDG/Source/NNERuntimeRDG/Private/Dml/NNEDmlOperatorUtils.h
--- a/Engine/Plugins/Experimental/NNERuntimeRDG/Source/NNERuntimeRDG/Private/Dml/NNEDmlOperatorUtils.h
+++ b/Engine/Plugins/Experimental/NNERuntimeRDG/Source/NNERuntimeRDG/Private/Dml/NNEDmlOperatorUtils.h
@@ -375,7 +375,7 @@ protected:
 
 public:
 
-	void Evaluate(TConstArrayView<uint32> InputShape, TConstArrayView<uint32> PaddingsValue = MakeEmptyConstArrayView<uint32>())
+	void Evaluate(TConstArrayView<uint32> InputShape, TConstArrayView<uint32> PaddingsValue = TConstArrayView<uint32>())
 	{
 		if (bIsGlobalKernel && WindowSize.IsEmpty())
 		{
