NNE Support more versions of pool ops

Problems:
* Some internal model required adding support of MaxPool-10.

Solutions:
* Add support for MaxPool-10. Newer versions of MaxPool and of AveragePool came for free given that attributes and inputs/outputs don't change.

Tests:
* NNEQA iteration tests
* NNEQA opsets tests for MaxPool and AveragePool
* Style transfer on win64 editor + standalone

#rnx

diff --git a/Engine/Plugins/Experimental/NNERuntimeRDG/Source/NNEHlslShaders/Shaders/NNEHlslShadersPool.usf b/Engine/Plugins/Experimental/NNERuntimeRDG/Source/NNEHlslShaders/Shaders/NNEHlslShadersPool.usf
--- a/Engine/Plugins/Experimental/NNERuntimeRDG/Source/NNEHlslShaders/Shaders/NNEHlslShadersPool.usf
+++ b/Engine/Plugins/Experimental/NNERuntimeRDG/Source/NNEHlslShaders/Shaders/NNEHlslShadersPool.usf
@@ -20,6 +20,7 @@ uint KernelVolume;
 #define SPATIAL_STRIDE 0
 #define SPATIAL_KERNEL 1
 #define SPATIAL_PAD_START 2
+#define SPATIAL_DILATION 3
 
 #define MAX_POOL 0
 #define AVERAGE_POOL 1
@@ -28,6 +29,7 @@ uint KernelVolume;
 void Pool(in const uint3 DispatchThreadID : SV_DispatchThreadID)
 {
 	const uint Index = DispatchThreadID.y * ThreadCountX + DispatchThreadID.x;
+	const uint SpatialDimOffset = 2;
 	
 	if (Index < Num)
 	{
@@ -46,9 +48,9 @@ void Pool(in const uint3 DispatchThreadID : SV_DispatchThreadID)
 		uint SpatialDimEnds[NUM_SPATIAL_DIMENSIONS];
 		for (uint spatialDim = 0; spatialDim < NUM_SPATIAL_DIMENSIONS; ++spatialDim)
 		{
-			const uint Dim = spatialDim+2;
+			const uint Dim = SpatialDimOffset + spatialDim;
 			const int SpatialDimStart = ((int)OutputIndices[Dim] * (int)SpatialInfo[spatialDim][SPATIAL_STRIDE]) - (int)SpatialInfo[spatialDim][SPATIAL_PAD_START];
-			const int SpatialDimEnd = SpatialDimStart + (int)SpatialInfo[spatialDim][SPATIAL_KERNEL];
+			const int SpatialDimEnd = SpatialDimStart + (int)SpatialInfo[spatialDim][SPATIAL_KERNEL] * (int)SpatialInfo[spatialDim][SPATIAL_DILATION];
 
 			SpatialDimStarts[spatialDim] = (uint)max(SpatialDimStart, (int)0);
 			SpatialDimEnds[spatialDim] = (uint)min(SpatialDimEnd, (int)TensorInfo[Dim][INPUT_SIZE]);
@@ -61,26 +63,26 @@ void Pool(in const uint3 DispatchThreadID : SV_DispatchThreadID)
 		float PooledValue = 0;
 		int Count = 0;
 		
-		for (SpatialDimIterator[0] = SpatialDimStarts[0]; SpatialDimIterator[0] < SpatialDimEnds[0]; ++SpatialDimIterator[0])
+		for (SpatialDimIterator[0] = SpatialDimStarts[0]; SpatialDimIterator[0] < SpatialDimEnds[0]; SpatialDimIterator[0] += SpatialInfo[0][SPATIAL_DILATION])
 		{
 		#if (NUM_SPATIAL_DIMENSIONS > 1)
-		for (SpatialDimIterator[1] = SpatialDimStarts[1]; SpatialDimIterator[1] < SpatialDimEnds[1]; ++SpatialDimIterator[1])
+		for (SpatialDimIterator[1] = SpatialDimStarts[1]; SpatialDimIterator[1] < SpatialDimEnds[1]; SpatialDimIterator[1] += SpatialInfo[1][SPATIAL_DILATION])
 		{
 		#endif
 		#if (NUM_SPATIAL_DIMENSIONS > 2)
-		for (SpatialDimIterator[2] = SpatialDimStarts[2]; SpatialDimIterator[2] < SpatialDimEnds[2]; ++SpatialDimIterator[2])
+		for (SpatialDimIterator[2] = SpatialDimStarts[2]; SpatialDimIterator[2] < SpatialDimEnds[2]; SpatialDimIterator[2] += SpatialInfo[2][SPATIAL_DILATION])
 		{
 		#endif
 		#if (NUM_SPATIAL_DIMENSIONS > 3)
-		for (SpatialDimIterator[3] = SpatialDimStarts[3]; SpatialDimIterator[3] < SpatialDimEnds[3]; ++SpatialDimIterator[3])
+		for (SpatialDimIterator[3] = SpatialDimStarts[3]; SpatialDimIterator[3] < SpatialDimEnds[3]; SpatialDimIterator[3] += SpatialInfo[3][SPATIAL_DILATION])
 		{
 		#endif
 		#if (NUM_SPATIAL_DIMENSIONS > 4)
-		for (SpatialDimIterator[4] = SpatialDimStarts[4]; SpatialDimIterator[4] < SpatialDimEnds[4]; ++SpatialDimIterator[4])
+		for (SpatialDimIterator[4] = SpatialDimStarts[4]; SpatialDimIterator[4] < SpatialDimEnds[4]; SpatialDimIterator[4] += SpatialInfo[4][SPATIAL_DILATION])
 		{
 		#endif
 		#if (NUM_SPATIAL_DIMENSIONS > 5)
-		for (SpatialDimIterator[5] = SpatialDimStarts[1]; SpatialDimIterator[5] < SpatialDimEnds[5]; ++SpatialDimIterator[5])
+		for (SpatialDimIterator[5] = SpatialDimStarts[5]; SpatialDimIterator[5] < SpatialDimEnds[5]; SpatialDimIterator[5] += SpatialInfo[5][SPATIAL_DILATION])
 		{
 		#endif
 
diff --git a/Engine/Plugins/Experimental/NNERuntimeRDG/Source/NNERuntimeRDG/Private/Dml/NNEDmlOperatorPool.cpp b/Engine/Plugins/Experimental/NNERuntimeRDG/Source/NNERuntimeRDG/Private/Dml/NNEDmlOperatorPool.cpp
--- a/Engine/Plugins/Experimental/NNERuntimeRDG/Source/NNERuntimeRDG/Private/Dml/NNEDmlOperatorPool.cpp
+++ b/Engine/Plugins/Experimental/NNERuntimeRDG/Source/NNERuntimeRDG/Private/Dml/NNEDmlOperatorPool.cpp
@@ -190,13 +190,20 @@ public:
 			TConstArrayView<uint32> StartPadding = Args.GetStartPadding();
 			TConstArrayView<uint32> EndPadding = Args.GetEndPadding();
 			TConstArrayView<uint32> Strides = Args.GetStrides();
+			TConstArrayView<uint32> Dilations = Args.GetDilations();
 			TConstArrayView<uint32> WindowSize = Args.GetWindowSize();
 			TConstArrayView<uint32> OutputShape = OutputTensors[Idx]->GetShape().GetData();
 
 			for (uint32 Dim = 0; Dim < Args.GetNumDimensions(); ++Dim)
 			{
 				const uint32 PaddedSize = InputShape[Dim + NonspatialDimensionCount] + StartPadding[Dim] + EndPadding[Dim];
-				check(OutputShape[Dim + NonspatialDimensionCount] == (PaddedSize - WindowSize[Dim]) / Strides[Dim] + 1);
+				float(*const RoundingFunction)(float) = 
+							Args.GetCeilMode() ?
+								(float(*)(float)) &FMath::CeilToFloat
+							:
+								(float(*)(float)) &FMath::FloorToFloat
+							;
+				check(OutputShape[Dim + NonspatialDimensionCount] == (uint32) RoundingFunction((float)(PaddedSize - ((WindowSize[Dim] - 1) * Dilations[Dim] + 1 )) / (float) Strides[Dim] + 1.0f));
 			}
 		}
 
@@ -313,6 +320,7 @@ static FDmlOperator##OpName##Version##Registrator RegisterDmlOperator##OpName##V
 // Register pooling operator on Module startup
 NNE_DML_REGISTER_POOLING_OP(MaxPool, MAX, false, 1)
 NNE_DML_REGISTER_POOLING_OP(MaxPool, MAX, false, 8)
+NNE_DML_REGISTER_POOLING_OP(MaxPool, MAX, false, 10)
 NNE_DML_REGISTER_POOLING_OP(MaxPool, MAX, false, 11)
 NNE_DML_REGISTER_POOLING_OP(MaxPool, MAX, false, 12)
 NNE_DML_REGISTER_POOLING_OP(GlobalMaxPool, MAX, true, 1)
diff --git a/Engine/Plugins/Experimental/NNERuntimeRDG/Source/NNERuntimeRDG/Private/Dml/NNEDmlOperatorUtils.h b/Engine/Plugins/Experimental/NNERuntimeRDG/Source/NNERuntimeRDG/Private/Dml/NNEDmlOperatorUtils.h
--- a/Engine/Plugins/Experimental/NNERuntimeRDG/Source/NNERuntimeRDG/Private/Dml/NNEDmlOperatorUtils.h
+++ b/Engine/Plugins/Experimental/NNERuntimeRDG/Source/NNERuntimeRDG/Private/Dml/NNEDmlOperatorUtils.h
@@ -260,6 +260,11 @@ public:
 		return OutputShape;
 	}
 
+	bool GetCeilMode() const
+	{
+		return bCeilMode;
+	}
+
 private:
 
 	FPaddingsHelper			Paddings;
@@ -278,6 +283,7 @@ protected:
 	bool					bIsGlobalKernel;
 	bool					bIsTransposed;
 	bool					bHasOutputShape;
+	bool					bCeilMode;
 
 protected:
 
@@ -368,6 +374,8 @@ protected:
 			{
 				return false;
 			}
+
+			bCeilMode = (bool) Attributes.GetValueOrDefault<int32>(TEXT("ceil_mode"), 0);
 		}
 
 		return true;
@@ -411,9 +419,15 @@ public:
 				checkf(Strides[Dim] != 0, TEXT("Strides must be != 0"));
 
 				uint32 StridableOutLen = PaddedLen - KernelLen;
-				uint32 OutLen = 1 + (StridableOutLen / Strides[Dim]);
-
-				OutputShape[Dim + DimOffset] = OutLen;
+				float OutLen = 1.0f + ((float) StridableOutLen / (float) Strides[Dim]);
+
+				float(*const RoundingFunction)(float) = 
+							bCeilMode?
+								(float(*)(float)) &FMath::CeilToFloat
+							:
+								(float(*)(float)) &FMath::FloorToFloat
+							;
+				OutputShape[Dim + DimOffset] = (uint32) RoundingFunction(OutLen);
 			}
 		}
 		else
diff --git a/Engine/Plugins/Experimental/NNERuntimeRDG/Source/NNERuntimeRDG/Private/Hlsl/NNERuntimeRDGPool.cpp b/Engine/Plugins/Experimental/NNERuntimeRDG/Source/NNERuntimeRDG/Private/Hlsl/NNERuntimeRDGPool.cpp
--- a/Engine/Plugins/Experimental/NNERuntimeRDG/Source/NNERuntimeRDG/Private/Hlsl/NNERuntimeRDGPool.cpp
+++ b/Engine/Plugins/Experimental/NNERuntimeRDG/Source/NNERuntimeRDG/Private/Hlsl/NNERuntimeRDGPool.cpp
@@ -20,6 +20,10 @@ namespace UE::NNERuntimeRDG::Private::Hlsl
 	template< UE::NNEHlslShaders::Internal::EPoolOperatorType PoolOperatorType >
 	class FPoolOperator : public FOperatorHlsl
 	{
+		static constexpr int32 SpatialInfoStrideIndex = 0;
+		static constexpr int32 SpatialInfoKernelIndex = 1;
+		static constexpr int32 SpatialInfoPadStartIndex = 2;
+		static constexpr int32 SpatialInfoDilationIndex = 3;
 	public:
 
 		FPoolOperator() {}
@@ -29,7 +33,9 @@ namespace UE::NNERuntimeRDG::Private::Hlsl
 		NNEHlslShaders::Internal::EConvAutoPad AutoPad = NNEHlslShaders::Internal::EConvAutoPad::NOTSET;
 		TArray<int32> Pads;
 		TArray<int32> Strides;
+		TArray<int32> Dilations;
 		TArray<int32> KernelShape;
+		int32 CeilMode = 0; // 0 is floor, 1 is ceil
 		int32 KernelVolume = 0;
 
 	public:
@@ -45,6 +51,7 @@ namespace UE::NNERuntimeRDG::Private::Hlsl
 
 			check(Pads.Num() == 2*NumSpatialDimensions);
 			check(Strides.Num() == NumSpatialDimensions);
+			check(Dilations.Num() == NumSpatialDimensions);
 			check(KernelShape.Num() == NumSpatialDimensions);
 
 			OutputShape.SetNumUninitialized(InputShape.Num());
@@ -58,12 +65,18 @@ namespace UE::NNERuntimeRDG::Private::Hlsl
 					case NNEHlslShaders::Internal::EConvAutoPad::NOTSET:
 					{
 						uint32 PadShape = Pads[i] + Pads[i+NumSpatialDimensions];
-						OutputShape[i + 2] = (uint32)FMath::FloorToFloat((float)(InputShape[i + 2] + PadShape - KernelShape[i]) / (float)Strides[i] + 1.0f);
+						float(*const RoundingFunction)(float) = 
+							CeilMode == 0 ?
+								(float(*)(float)) &FMath::FloorToFloat
+							:
+								(float(*)(float)) &FMath::CeilToFloat
+							;
+						OutputShape[i + 2] = (uint32)RoundingFunction((float)(InputShape[i + 2] + PadShape - ((KernelShape[i] - 1) * Dilations[i] + 1)) / (float)Strides[i] + 1.0f);
 						break;
 					}
 					case NNEHlslShaders::Internal::EConvAutoPad::VALID:
 					{
-						OutputShape[i + 2] = (uint32)FMath::CeilToFloat((float)(InputShape[i + 2] - KernelShape[i] + 1) / (float)Strides[i]);
+						OutputShape[i + 2] = (uint32)FMath::CeilToFloat((float)(InputShape[i + 2] - ((KernelShape[i] - 1) * Dilations[i] + 1) + 1) / (float)Strides[i]);
 						break;
 					}
 					case NNEHlslShaders::Internal::EConvAutoPad::SAME_UPPER:
@@ -109,14 +122,19 @@ namespace UE::NNERuntimeRDG::Private::Hlsl
 
 			TArray<int32> StridesDefault;
 			TArray<int32> PadsDefault;
+			TArray<int32> DilationsDefault;
 			StridesDefault.Init(1, NumSpatialDimensions);
 			PadsDefault.Init(0, 2 * NumSpatialDimensions);
+			DilationsDefault.Init(1, NumSpatialDimensions);
 
 			NNEHlslShaders::Internal::FConvCS::LexFromString(AutoPad, *Attributes.GetValueOrDefault<FString>(TEXT("auto_pad"), TEXT("NOTSET")));
 			Pads = Attributes.GetValueOrDefault<TArray<int32>>(TEXT("pads"), PadsDefault);
 			Strides = Attributes.GetValueOrDefault<TArray<int32>>(TEXT("strides"), StridesDefault);
+			Dilations = Attributes.GetValueOrDefault<TArray<int32>>(TEXT("dilations"), DilationsDefault);
 			KernelShape = Attributes.GetValue<TArray<int32>>(TEXT("kernel_shape"));
 
+			CeilMode = Attributes.GetValueOrDefault<int32>(TEXT("ceil_mode"), CeilMode);
+
 			if (KernelShape.Num() != NumSpatialDimensions)
 			{
 				UE_LOG(LogNNE, Warning, TEXT("%s KernelShape should have as many elements as the spatial dimensions of the input, got %d while input have %d."), GetOperatorName(), KernelShape.Num(), NumSpatialDimensions);
@@ -127,6 +145,11 @@ namespace UE::NNERuntimeRDG::Private::Hlsl
 				UE_LOG(LogNNE, Warning, TEXT("%s Strides should have as many elements as the spatial dimensions of the input, got %d while input have %d."), GetOperatorName(), Strides.Num(), NumSpatialDimensions);
 				return false;
 			}
+			if (Dilations.Num() != NumSpatialDimensions)
+			{
+				UE_LOG(LogNNE, Warning, TEXT("%s Dilations should have as many elements as the spatial dimensions of the input, got %d while input have %d."), GetOperatorName(), Dilations.Num(), NumSpatialDimensions);
+				return false;
+			}
 			if (Pads.Num() != 2*NumSpatialDimensions)
 			{
 				UE_LOG(LogNNE, Warning, TEXT("%s Pads should have twice as many elements as the spatial dimensions of the input, got %d while input have %d."), GetOperatorName(), Pads.Num(), NumSpatialDimensions);
@@ -177,18 +200,19 @@ namespace UE::NNERuntimeRDG::Private::Hlsl
 			FillTensorSizeShaderParameters(Input, Params->TensorInfo, 2);
 			for (int32 i = 0; i < NumSpatialDimensions; ++i)
 			{
-				Params->SpatialInfo[i][0] = Strides[i];
-				Params->SpatialInfo[i][1] = KernelShape[i];
+				Params->SpatialInfo[i][SpatialInfoStrideIndex] = Strides[i];
+				Params->SpatialInfo[i][SpatialInfoKernelIndex] = KernelShape[i];
 				if (AutoPad == NNEHlslShaders::Internal::EConvAutoPad::SAME_LOWER)
 				{
 					// see [REDACTED_PROFILE_URL]/onnx/blob/main/docs/Changelog.md#MaxPool-8
 					// only needed for SAME_LOWER as SAME_UPPER is handled by index clamping in the HLSL kernel
-					Params->SpatialInfo[i][2] = (Output.GetShape().GetData()[i + 2] - 1) * Strides[i] + KernelShape[i] - Input.GetShape().GetData()[i + 2];
+					Params->SpatialInfo[i][SpatialInfoPadStartIndex] = (Output.GetShape().GetData()[i + 2] - 1) * Strides[i] + ((KernelShape[i] - 1) * Dilations[i] + 1) - Input.GetShape().GetData()[i + 2];
 				}
 				else
 				{
-					Params->SpatialInfo[i][2] = Pads[i];
+					Params->SpatialInfo[i][SpatialInfoPadStartIndex] = Pads[i];
 				}
+				Params->SpatialInfo[i][SpatialInfoDilationIndex] = Dilations[i];
 			}
 
 			FPoolCS::FPermutationDomain PermutationVector;
@@ -247,8 +271,10 @@ namespace UE::NNERuntimeRDG::Private::Hlsl
 		AttributeValidator.AddOptional(TEXT("auto_pad"), ENNEAttributeDataType::String);
 		AttributeValidator.AddRequired(TEXT("kernel_shape"), ENNEAttributeDataType::Int32Array);
 		AttributeValidator.AddOptional(TEXT("pads"), ENNEAttributeDataType::Int32Array);
+		AttributeValidator.AddOptional(TEXT("ceil_mode"), ENNEAttributeDataType::Int32);
 		if constexpr (PoolOperatorType == UE::NNEHlslShaders::Internal::EPoolOperatorType::MAX_POOL)
 		{
+			AttributeValidator.AddOptional(TEXT("dilations"), ENNEAttributeDataType::Int32Array);
 			AttributeValidator.AddOptional(TEXT("storage_order"), ENNEAttributeDataType::Int32);//Unused, only needed for 2nd output itself not supported, see [REDACTED_PROFILE_URL]/onnx/issues/1370
 		}
 		else
@@ -281,7 +307,13 @@ namespace UE::NNERuntimeRDG::Private::Hlsl
 	{
 		// Note: support of a particular version is partial with respect to tensor data types (only the most typical ones are usually supported).
 		Registry.OpAdd({{TEXT("MaxPool"), TEXT("Onnx")}, 8}, CreateMaxPoolOperator, ValidatePoolOperator<UE::NNEHlslShaders::Internal::EPoolOperatorType::MAX_POOL>);
+		Registry.OpAdd({{TEXT("MaxPool"), TEXT("Onnx")}, 10}, CreateMaxPoolOperator, ValidatePoolOperator<UE::NNEHlslShaders::Internal::EPoolOperatorType::MAX_POOL>);
+		Registry.OpAdd({{TEXT("MaxPool"), TEXT("Onnx")}, 11}, CreateMaxPoolOperator, ValidatePoolOperator<UE::NNEHlslShaders::Internal::EPoolOperatorType::MAX_POOL>);
+		Registry.OpAdd({{TEXT("MaxPool"), TEXT("Onnx")}, 12}, CreateMaxPoolOperator, ValidatePoolOperator<UE::NNEHlslShaders::Internal::EPoolOperatorType::MAX_POOL>);
 		Registry.OpAdd({ {TEXT("AveragePool"), TEXT("Onnx")}, 7}, CreateAveragePoolOperator, ValidatePoolOperator<UE::NNEHlslShaders::Internal::EPoolOperatorType::AVERAGE_POOL>);
+		Registry.OpAdd({ {TEXT("AveragePool"), TEXT("Onnx")}, 10}, CreateAveragePoolOperator, ValidatePoolOperator<UE::NNEHlslShaders::Internal::EPoolOperatorType::AVERAGE_POOL>);
+		Registry.OpAdd({ {TEXT("AveragePool"), TEXT("Onnx")}, 11}, CreateAveragePoolOperator, ValidatePoolOperator<UE::NNEHlslShaders::Internal::EPoolOperatorType::AVERAGE_POOL>);
+		Registry.OpAdd({ {TEXT("AveragePool"), TEXT("Onnx")}, 19}, CreateAveragePoolOperator, ValidatePoolOperator<UE::NNEHlslShaders::Internal::EPoolOperatorType::AVERAGE_POOL>);
 		return true;
 	}
 } // UE::NNERuntimeRDG::Private::Hlsl
