[NNE] NNERuntimeRDGHlsl : Fix leak of operator instances

Problem: NNE RDG Hlsl was leaking operator instances
Solution: Release related memory on model instance destruction.

#tests win64 editor: smoke test, win64 standalone: all tests + style transfer, checked memory using unreal insight.

#rnx

diff --git a/Engine/Plugins/Experimental/NNERuntimeRDG/Source/NNERuntimeRDG/Private/NNERuntimeRDGModelHlsl.cpp b/Engine/Plugins/Experimental/NNERuntimeRDG/Source/NNERuntimeRDG/Private/NNERuntimeRDGModelHlsl.cpp
--- a/Engine/Plugins/Experimental/NNERuntimeRDG/Source/NNERuntimeRDG/Private/NNERuntimeRDGModelHlsl.cpp
+++ b/Engine/Plugins/Experimental/NNERuntimeRDG/Source/NNERuntimeRDG/Private/NNERuntimeRDGModelHlsl.cpp
@@ -104,6 +104,15 @@ bool FModelInstance::PrepareModelRDG(FRDGBuilder& RDGBuilder)
 	return true;
 }
 
+FModelInstance::~FModelInstance()
+{
+	for (FOperatorHlsl* Operator : Operators)
+	{
+		delete Operator;
+	}
+	Operators.Empty();
+}
+
 bool FModelInstance::Init(TConstArrayView<uint8> ModelData)
 {
 	check(ModelData.Num() > 0);
diff --git a/Engine/Plugins/Experimental/NNERuntimeRDG/Source/NNERuntimeRDG/Private/NNERuntimeRDGModelHlsl.h b/Engine/Plugins/Experimental/NNERuntimeRDG/Source/NNERuntimeRDG/Private/NNERuntimeRDGModelHlsl.h
--- a/Engine/Plugins/Experimental/NNERuntimeRDG/Source/NNERuntimeRDG/Private/NNERuntimeRDGModelHlsl.h
+++ b/Engine/Plugins/Experimental/NNERuntimeRDG/Source/NNERuntimeRDG/Private/NNERuntimeRDGModelHlsl.h
@@ -17,7 +17,7 @@ class FModelInstance : public FModelInstanceRDG
 	
 public:
 
-	~FModelInstance() = default;
+	~FModelInstance();
 
 	bool Init(TConstArrayView<uint8> ModelData);
 
