[NNE] NNERuntimeRDGDml fix infinite wait loop on console.

Problem:
Any time the NNERuntimeRDGDml was used on a console the App immediately froze.

Solution:
Removed the loop that caused an indefinite wait.

#tests Unreal Game on console, style-transfer and all tests

#rnx

diff --git a/Engine/Plugins/Experimental/NNERuntimeRDG/Source/NNERuntimeRDG/Private/Dml/NNEDmlModel.cpp b/Engine/Plugins/Experimental/NNERuntimeRDG/Source/NNERuntimeRDG/Private/Dml/NNEDmlModel.cpp
--- a/Engine/Plugins/Experimental/NNERuntimeRDG/Source/NNERuntimeRDG/Private/Dml/NNEDmlModel.cpp
+++ b/Engine/Plugins/Experimental/NNERuntimeRDG/Source/NNERuntimeRDG/Private/Dml/NNEDmlModel.cpp
@@ -844,17 +844,6 @@ FModelInstance::~FModelInstance()
 		{
 			NNE_TRACE_EVENT_SCOPED(NNE_DmlModel_ReleaseResources_RT);
 
-			if (DispatchFence.IsValid())
-			{
-				while (!DispatchFence->Poll())
-				{
-					FPlatformProcess::Sleep(0.0f);
-				}
-
-				DispatchFence->DisableLifetimeExtension();
-				DispatchFence.SafeRelease();
-			}
-
 			DescHeap.Reset();
 			BindingTable.Reset();
 
@@ -1055,8 +1044,6 @@ bool FModelInstance::InitCompiledOp()
 		{
 			NNE_TRACE_EVENT_SCOPED(NNE_DmlModelInstance_InitCompiledOp_RT);
 
-			DispatchFence = DynamicRHI->RHICreateGPUFence(TEXT("NNE_DmlModelInstance_DispatchFence"));
-
 			FRHIBufferInputArray	Inputs;
 
 			for (int32 InputIdx : InputTensorIndices)
@@ -1137,7 +1124,6 @@ bool FModelInstance::InitCompiledOp()
 				check(InitTempBuff.IsValid());
 			}
 
-			DispatchFence->Clear();
 			RHICmdList.EnqueueLambda(
 				[this, Inputs = MoveTemp(Inputs), Barriers = MoveTemp(Barriers), InitTempBuff = MoveTemp(InitTempBuff)](FRHICommandListImmediate& RHICmdList)
 				{
@@ -1162,24 +1148,9 @@ bool FModelInstance::InitCompiledOp()
 					DynamicRHI->RHIFinishExternalComputeWork(DevCtx->DeviceIndex, D3DCmdList);
 				}
 			);
-			RHICmdList.WriteGPUFence(DispatchFence);
-			
-			RHICmdList.ImmediateFlush(EImmediateFlushType::FlushRHIThread);
 
-			{
-				NNE_TRACE_EVENT_SCOPED(NNE_DmlModelInstance_InitCompiledOp_Wait_RT)
-				
-				while (!DispatchFence->Poll())
-				{
-					FPlatformProcess::Sleep(0.0f);
-				}
-				
-				// Potentially we can create a lot of weight buffers, we should release them as soon as possible
-				if (MemSizeWeights)
-				{
-					RHICmdList.ImmediateFlush(EImmediateFlushType::FlushRHIThreadFlushResources);
-				}
-			}
+			// Make sure everything is submitted to the GPU before returning
+			RHICmdList.ImmediateFlush(EImmediateFlushType::FlushRHIThread);
 
 			Signal->Trigger();
 		}
@@ -1267,7 +1238,6 @@ void FModelInstance::AddDispatchOps_RenderThread(FRDGBuilder& GraphBuilder)
 			{
 				SCOPED_GPU_STAT(RHICmdList, NNE_DmlModelInstance_DispatchD3D_GPU)
 
-				DispatchFence->Clear();
 				RHICmdList.EnqueueLambda(
 					[this, InputBuffers = MoveTemp(RHIInputBuffers), OutputBuffers = MoveTemp(RHIOutputBuffers), TempBuff = MoveTemp(TempBuff)](FRHICommandListImmediate& RHICmdList)
 					{
@@ -1332,8 +1302,6 @@ void FModelInstance::AddDispatchOps_RenderThread(FRDGBuilder& GraphBuilder)
 						DynamicRHI->RHIFinishExternalComputeWork(DevCtx->DeviceIndex, D3DCmdList);
 					}
 				);
-
-				RHICmdList.WriteGPUFence(DispatchFence);
 			}
 		}
 	);
diff --git a/Engine/Plugins/Experimental/NNERuntimeRDG/Source/NNERuntimeRDG/Private/Dml/NNEDmlModel.h b/Engine/Plugins/Experimental/NNERuntimeRDG/Source/NNERuntimeRDG/Private/Dml/NNEDmlModel.h
--- a/Engine/Plugins/Experimental/NNERuntimeRDG/Source/NNERuntimeRDG/Private/Dml/NNEDmlModel.h
+++ b/Engine/Plugins/Experimental/NNERuntimeRDG/Source/NNERuntimeRDG/Private/Dml/NNEDmlModel.h
@@ -118,7 +118,6 @@ private:
 	uint32								DescCount;
 	uint32								DescSize;
 
-	FGPUFenceRHIRef						DispatchFence;
 	FBufferRHIRef						PersistBuff;
 	uint64								MemSizeWeights;
 	uint64								MemSizeTemp;
