[NNE] Fix: NNERuntimeRDGHlsl Maxpool fix for SM6

Problems:
* HLSL Maxpool op fail on SM6

Solution:
* Fix it

#tests win64 standalone: all hlsl tests, console: pool ops hlsl tests

#rnx

diff --git a/Engine/Plugins/Experimental/NNERuntimeRDG/Source/NNEHlslShaders/Shaders/NNEHlslShadersPool.usf b/Engine/Plugins/Experimental/NNERuntimeRDG/Source/NNEHlslShaders/Shaders/NNEHlslShadersPool.usf
--- a/Engine/Plugins/Experimental/NNERuntimeRDG/Source/NNEHlslShaders/Shaders/NNEHlslShadersPool.usf
+++ b/Engine/Plugins/Experimental/NNERuntimeRDG/Source/NNEHlslShaders/Shaders/NNEHlslShadersPool.usf
@@ -94,7 +94,9 @@ void Pool(in const uint3 DispatchThreadID : SV_DispatchThreadID)
 			}
 			float CurrentValue = Input[InputIndex];
 			#if (POOL_OPERATOR_TYPE == MAX_POOL)
-				PooledValue = (Count == 0) ? CurrentValue : max(PooledValue, CurrentValue);
+				//float MaxValue = max(PooledValue, CurrentValue); //Is not working properly on DX12 SM6 thus the ternary op for MaxValue
+				float MaxValue = (PooledValue > CurrentValue) ? PooledValue : CurrentValue;
+				PooledValue = (Count == 0) ? CurrentValue : MaxValue;
 			#else //POOL_OPERATOR_TYPE == AVERAGE_POOL
 				PooledValue += CurrentValue;
 			#endif 
