[NNE] NNERuntimeRDGDml fix (again) GPU crash.

#tests Editor/Game (Win64) style-transfer and all tests, Game on consoles all tests

#rnx

diff --git a/Engine/Plugins/Experimental/NNERuntimeRDG/Source/NNERuntimeRDG/Private/Dml/NNEDmlModel.cpp b/Engine/Plugins/Experimental/NNERuntimeRDG/Source/NNERuntimeRDG/Private/Dml/NNEDmlModel.cpp
--- a/Engine/Plugins/Experimental/NNERuntimeRDG/Source/NNERuntimeRDG/Private/Dml/NNEDmlModel.cpp
+++ b/Engine/Plugins/Experimental/NNERuntimeRDG/Source/NNERuntimeRDG/Private/Dml/NNEDmlModel.cpp
@@ -844,14 +844,9 @@ FModelInstance::~FModelInstance()
 		{
 			NNE_TRACE_EVENT_SCOPED(NNE_DmlModel_ReleaseResources_RT);
 
-			DescHeap.Reset();
-			BindingTable.Reset();
-
-			if (PersistBuff.IsValid())
-			{
-				PersistBuff->DisableLifetimeExtension();
-				PersistBuff.SafeRelease();
-			}
+			// Need to wait for GPU to finish to not release resources possibly still in use
+			// This needs proper rework!!!
+			RHICmdList.BlockUntilGPUIdle();
 
 			Signal->Trigger();
 		}
@@ -864,7 +859,6 @@ FModelInstance::~FModelInstance()
 	{
 		delete Op;
 	}
-	Operators.Reset();
 }
 
 bool FModelInstance::Init(TConstArrayView<uint8> ModelData, FDmlDeviceContext* InDevCtx)
@@ -1148,9 +1142,10 @@ bool FModelInstance::InitCompiledOp()
 					DynamicRHI->RHIFinishExternalComputeWork(DevCtx->DeviceIndex, D3DCmdList);
 				}
 			);
-
-			// Make sure everything is submitted to the GPU before returning
-			RHICmdList.ImmediateFlush(EImmediateFlushType::FlushRHIThread);
+			
+			// Wait for GPU to finish to avoid potential crash when dispatching operators
+			// This needs proper rework!!!
+			RHICmdList.BlockUntilGPUIdle();
 
 			Signal->Trigger();
 		}
