[NNE] NNERuntimeRDGDml added more detailed description of tentative fix for GPU crash for future rework.

#tests none

#rnx

diff --git a/Engine/Plugins/Experimental/NNERuntimeRDG/Source/NNERuntimeRDG/Private/Dml/NNEDmlModel.cpp b/Engine/Plugins/Experimental/NNERuntimeRDG/Source/NNERuntimeRDG/Private/Dml/NNEDmlModel.cpp
--- a/Engine/Plugins/Experimental/NNERuntimeRDG/Source/NNERuntimeRDG/Private/Dml/NNEDmlModel.cpp
+++ b/Engine/Plugins/Experimental/NNERuntimeRDG/Source/NNERuntimeRDG/Private/Dml/NNEDmlModel.cpp
@@ -1143,8 +1143,8 @@ bool FModelInstance::InitCompiledOp()
 				}
 			);
 			
-			// Wait for GPU to finish to avoid potential crash when dispatching operators
-			// This needs proper rework!!!
+			// Since the operator initializer and the compiled operator share GPU resources (currently the descriptor heap),
+			// we have to wait for the initializer on the GPU to finish, before we can dispatch the compiled operator!
 			RHICmdList.BlockUntilGPUIdle();
 
 			Signal->Trigger();
