[NNE] Add Hlsl Split implementation

Problem:
* Hlsl Split implementation is missing

Solution:
* Add implementation and enable testing on Hlsl

Tests:
* NNEQA Hlsl tests (iteration + opsets)
* Win64 Standalone

#rnx

#virtualized

diff --git a/Engine/Plugins/Experimental/NNERuntimeRDG/Source/NNEHlslShaders/Internal/NNEHlslShadersSplitCS.h b/Engine/Plugins/Experimental/NNERuntimeRDG/Source/NNEHlslShaders/Internal/NNEHlslShadersSplitCS.h
new file mode 100644
--- /dev/null
+++ b/Engine/Plugins/Experimental/NNERuntimeRDG/Source/NNEHlslShaders/Internal/NNEHlslShadersSplitCS.h
@@ -0,0 +1,45 @@
+// Copyright Epic Games, Inc. All Rights Reserved.
+
+#pragma once
+
+#include "NNEHlslShaderBase.h"
+#include "RenderGraphUtils.h"
+#include "ShaderParameterUtils.h"
+
+namespace UE::NNEHlslShaders::Internal
+{
+
+	class FSplitConstants
+	{
+	public:
+
+		static const int32 MAX_NUM_DIMENSIONS{ 8 };
+		static const int32 MAX_NUM_SPLITS{ 8 };
+		static const int32 NUM_GROUP_THREADS{ 256 };
+	};
+
+	class NNEHLSLSHADERS_API FSplitCS : public FHlslShaderBase
+	{
+		DECLARE_GLOBAL_SHADER(FSplitCS);
+		SHADER_USE_PARAMETER_STRUCT(FSplitCS, FHlslShaderBase)
+
+		class FSplitRank : SHADER_PERMUTATION_RANGE_INT("RANK", 1, FSplitConstants::MAX_NUM_DIMENSIONS);
+		class FSplitAxis : SHADER_PERMUTATION_RANGE_INT("AXIS", 0, FSplitConstants::MAX_NUM_DIMENSIONS - 1);
+		class FSplitNumSplits : SHADER_PERMUTATION_RANGE_INT("NUM_SPLITS", 1, FSplitConstants::MAX_NUM_SPLITS);
+		using FPermutationDomain = TShaderPermutationDomain<FSplitRank, FSplitAxis, FSplitNumSplits>;
+
+	public:
+
+		BEGIN_SHADER_PARAMETER_STRUCT(FParameters, )
+			SHADER_PARAMETER_RDG_BUFFER_SRV(Buffer<float>, Input)
+			SHADER_PARAMETER_RDG_BUFFER_UAV_ARRAY(RWBuffer<float>, Output, [FSplitConstants::MAX_NUM_SPLITS])
+			SHADER_PARAMETER_ARRAY(FUintVector4, InputTensorInfo, [FSplitConstants::MAX_NUM_DIMENSIONS])
+			SHADER_PARAMETER_ARRAY(FUintVector4, OutputTensorInfo, [FSplitConstants::MAX_NUM_SPLITS * FSplitConstants::MAX_NUM_DIMENSIONS])
+			SHADER_PARAMETER(uint32, ThreadCountX)
+            SHADER_PARAMETER(uint32, Num)
+		END_SHADER_PARAMETER_STRUCT()
+
+		static bool ShouldCompilePermutation(const FGlobalShaderPermutationParameters& InParameters);
+		static void ModifyCompilationEnvironment(const FGlobalShaderPermutationParameters& InParameters, FShaderCompilerEnvironment& OutEnvironment);
+	};
+} // UE::NNEHlslShaders::Internal
diff --git a/Engine/Plugins/Experimental/NNERuntimeRDG/Source/NNEHlslShaders/Private/NNEHlslShadersSplitCS.cpp b/Engine/Plugins/Experimental/NNERuntimeRDG/Source/NNEHlslShaders/Private/NNEHlslShadersSplitCS.cpp
new file mode 100644
--- /dev/null
+++ b/Engine/Plugins/Experimental/NNERuntimeRDG/Source/NNEHlslShaders/Private/NNEHlslShadersSplitCS.cpp
@@ -0,0 +1,27 @@
+// Copyright Epic Games, Inc. All Rights Reserved.
+
+#include "NNEHlslShadersSplitCS.h"
+#include "NNE.h"
+
+namespace UE::NNEHlslShaders::Internal
+{
+	bool FSplitCS::ShouldCompilePermutation(const FGlobalShaderPermutationParameters& InParameters)
+	{
+		const FPermutationDomain PermutationVector(InParameters.PermutationId);
+
+		return PermutationVector.Get<FSplitCS::FSplitAxis>() < PermutationVector.Get<FSplitCS::FSplitRank>();
+	}
+
+	void FSplitCS::ModifyCompilationEnvironment(const FGlobalShaderPermutationParameters& InParameters, FShaderCompilerEnvironment& OutEnvironment)
+	{
+		FGlobalShader::ModifyCompilationEnvironment(InParameters, OutEnvironment);
+
+		OutEnvironment.SetDefine(TEXT("THREADGROUP_SIZE_X"), FSplitConstants::NUM_GROUP_THREADS);
+		OutEnvironment.SetDefine(TEXT("MAX_NUM_SPLITS"), FSplitConstants::MAX_NUM_SPLITS);
+		OutEnvironment.SetDefine(TEXT("MAX_NUM_DIMENSIONS"), FSplitConstants::MAX_NUM_DIMENSIONS);
+
+		FPermutationDomain PermutationVector(InParameters.PermutationId);
+	}
+
+	IMPLEMENT_GLOBAL_SHADER(FSplitCS, "/NNE/NNEHlslShadersSplit.usf", "Split", SF_Compute);
+} // UE::NNEHlslShaders::Internal
\ No newline at end of file
diff --git a/Engine/Plugins/Experimental/NNERuntimeRDG/Source/NNEHlslShaders/Shaders/NNEHlslShadersSplit.usf b/Engine/Plugins/Experimental/NNERuntimeRDG/Source/NNEHlslShaders/Shaders/NNEHlslShadersSplit.usf
new file mode 100644
--- /dev/null
+++ b/Engine/Plugins/Experimental/NNERuntimeRDG/Source/NNEHlslShaders/Shaders/NNEHlslShadersSplit.usf
@@ -0,0 +1,130 @@
+// Copyright Epic Games, Inc. All Rights Reserved.
+
+#include "/Engine/Public/Platform.ush"
+#include "NNEHlslShadersBroadcastHelper.ush"
+
+Buffer<float> Input;
+
+// Buffer array is not currently available, requires a workaround. Also can't write into more than 8 output buffers.
+// RWBuffer<float> Output[MAX_NUM_SPLITS];
+#define TEMPLATE_BUFFERS(TPrefix) \
+	TPrefix##_0;\
+	TPrefix##_1;\
+	TPrefix##_2;\
+	TPrefix##_3;\
+	TPrefix##_4;\
+	TPrefix##_5;\
+	TPrefix##_6;\
+	TPrefix##_7;
+
+TEMPLATE_BUFFERS(RWBuffer<float> Output)
+
+uint4 InputTensorInfo[MAX_NUM_DIMENSIONS];
+uint4 OutputTensorInfo[MAX_NUM_SPLITS * MAX_NUM_DIMENSIONS];
+
+uint Num;
+uint ThreadCountX;
+
+#define STRIDE_IDX 0
+#define SIZE_IDX 1
+
+#define STATIC_LOOP(Var, From, To) \
+[unroll] \
+for(uint Var = From; Var < To; ++Var)
+
+void GetInputDimIterator(const uint GlobalIdx, out uint DimIterator[RANK])
+{
+	uint Offset = GlobalIdx;
+	STATIC_LOOP(DimIdx, 0, RANK)
+	{
+		uint Remainder;
+		DivMod(Offset, InputTensorInfo[DimIdx][STRIDE_IDX], DimIterator[DimIdx], Remainder);
+		Offset = Remainder;
+	}
+}
+
+struct FOutputIterator
+{
+	uint TensorIdx;
+	uint AxisIterator;
+};
+
+FOutputIterator GetOutputIterator(const uint InputAxisIterator)
+{
+	FOutputIterator OutputIterator;
+	uint SplitIndex = 0;
+	uint PrevSplitIndex = 0;
+
+	[loop]
+	for(OutputIterator.TensorIdx = 0; OutputIterator.TensorIdx < NUM_SPLITS; ++OutputIterator.TensorIdx)
+	{
+		SplitIndex += OutputTensorInfo[OutputIterator.TensorIdx * MAX_NUM_DIMENSIONS + AXIS][SIZE_IDX];
+		if(InputAxisIterator < SplitIndex)
+		{
+			OutputIterator.AxisIterator = InputAxisIterator - PrevSplitIndex;
+			return OutputIterator;
+		}
+		PrevSplitIndex = SplitIndex;
+	}
+	
+	return OutputIterator;
+}
+
+uint GetOutputGlobalIdx(const uint DimIterator[RANK], const uint OutputTensorIdx)
+{
+	uint Offset = 0;
+	STATIC_LOOP(DimIdx, 0, RANK)
+	{
+		Offset += DimIterator[DimIdx] * OutputTensorInfo[OutputTensorIdx * MAX_NUM_DIMENSIONS + DimIdx][STRIDE_IDX];
+	}
+	return Offset;
+}
+
+[numthreads(THREADGROUP_SIZE_X, 1, 1)]
+void Split(in const uint3 DispatchThreadID : SV_DispatchThreadID)
+{
+	const uint Index = DispatchThreadID.y * ThreadCountX + DispatchThreadID.x;
+	
+	if (Index < Num)
+	{
+		uint DimIterator[RANK];
+
+		GetInputDimIterator(Index, DimIterator);
+		
+		FOutputIterator OutputIterator = GetOutputIterator(DimIterator[AXIS]);
+		DimIterator[AXIS] = OutputIterator.AxisIterator;
+
+		// The following implements:
+		// Output[OutputIterator.TensorIdx][GetOutputGlobalIdx(DimIterator, OutputIterator.TensorIdx)] = Input[Index];
+
+		#define WRITE_BUFFER_ARRAY(Idx) \
+		if(OutputIterator.TensorIdx == Idx) \
+		{ \
+			Output_##Idx[GetOutputGlobalIdx(DimIterator, Idx)] = Input[Index]; \
+		}
+		#if 0 < NUM_SPLITS
+			WRITE_BUFFER_ARRAY(0)
+		#endif
+		#if 1 < NUM_SPLITS
+			WRITE_BUFFER_ARRAY(1)
+		#endif
+		#if 2 < NUM_SPLITS
+			WRITE_BUFFER_ARRAY(2)
+		#endif
+		#if 3 < NUM_SPLITS
+			WRITE_BUFFER_ARRAY(3)
+		#endif
+		#if 4 < NUM_SPLITS
+			WRITE_BUFFER_ARRAY(4)
+		#endif
+		#if 5 < NUM_SPLITS
+			WRITE_BUFFER_ARRAY(5)
+		#endif
+		#if 6 < NUM_SPLITS
+			WRITE_BUFFER_ARRAY(6)
+		#endif
+		#if 7 < NUM_SPLITS
+			WRITE_BUFFER_ARRAY(7)
+		#endif
+	}
+}
diff --git a/Engine/Plugins/Experimental/NNERuntimeRDG/Source/NNERuntimeRDG/Private/Hlsl/NNERuntimeRDGSplit.cpp b/Engine/Plugins/Experimental/NNERuntimeRDG/Source/NNERuntimeRDG/Private/Hlsl/NNERuntimeRDGSplit.cpp
new file mode 100644
--- /dev/null
+++ b/Engine/Plugins/Experimental/NNERuntimeRDG/Source/NNERuntimeRDG/Private/Hlsl/NNERuntimeRDGSplit.cpp
@@ -0,0 +1,249 @@
+// Copyright Epic Games, Inc. All Rights Reserved.
+
+#include "NNERuntimeRDGSplit.h"
+#include "NNEHlslShadersSplitCS.h"
+#include "NNERuntimeRDGHlslHelper.h"
+#include "NNETensor.h"
+#include "NNETypes.h"
+#include "RenderGraphUtils.h"
+#include "Algo/Accumulate.h"
+
+namespace UE::NNERuntimeRDG::Private::Hlsl
+{
+	DECLARE_GPU_STAT_NAMED(FNNEOperatorSplit, TEXT("NNE.Operator.Hlsl.Split"));
+
+	
+	/**
+	 * Split operator implementation
+	 */
+	template<int32 Version>
+	class FSplitOperator : public FOperatorHlsl
+	{
+	public:
+
+		FSplitOperator() {}
+		virtual ~FSplitOperator() = default;
+	private:
+		int32 Axis = 0;
+        TArray<int64> Splits;
+
+		bool CheckSplitsSum(const NNE::FTensorShape& InputShape)
+		{
+			int64 SumToCheck = Algo::Accumulate(Splits, /*Init*/ 0);
+			return InputShape.GetData()[Axis] == (uint32) SumToCheck;
+		}
+
+	public:
+
+		virtual int PrepareOutputs(TConstArrayView<NNE::Internal::FTensorRef> InputTensors, TArrayView<NNE::Internal::FTensorRef> OutputTensors) override
+		{
+			using namespace UE::NNEHlslShaders::Internal;
+
+			check(InputTensors.Num() >= 1 && InputTensors.Num() <= 2);
+			check(OutputTensors.Num() >= 1);
+
+			const NNE::FTensorShape& InputShape = InputTensors[0]->GetShape();
+
+			if(Version >= 13 && InputTensors.Num() == 2)
+			{
+				const NNE::Internal::FTensor& SplitTensor = *InputTensors[1];
+				if(!SplitTensor.HasPreparedData())
+				{
+					UE_LOG(LogNNE, Warning, TEXT("Split: RDG Hlsl requires the split tensor to be CPU-constant in order to compute output shapes."));
+					return -1;
+				}
+				
+				Splits = SplitTensor.GetPreparedData<int64>();
+			}
+
+			// Default: equal split
+			if(Splits.IsEmpty())
+			{
+				Splits.Init(InputShape.GetData()[Axis] / OutputTensors.Num(), OutputTensors.Num());
+				Splits.Last() += InputShape.GetData()[Axis] % OutputTensors.Num();
+			}
+
+			if(!CheckSplitsSum(InputShape))
+			{
+				UE_LOG(LogNNE, Warning, TEXT("Split: Sum of split values not equal to split axis' dimension."));
+				return -1;
+			}
+
+			if(Splits.Num() > FSplitConstants::MAX_NUM_SPLITS)
+			{
+				UE_LOG(LogNNE, Warning, TEXT("Split: Number of splits (%d) exceeds maximum allowed (%d)."), Splits.Num(), FSplitConstants::MAX_NUM_SPLITS);
+				return -1;
+			}
+			
+			if(OutputTensors.Num() != Splits.Num())
+			{
+				UE_LOG(LogNNE, Warning, TEXT("Split: Number of output tensors differs from number of splits provided."));
+				return -1;
+			}
+
+			for(int32 Idx = 0; Idx < OutputTensors.Num(); Idx++)
+			{
+				TArray<uint32, TInlineAllocator<NNE::FTensorShape::MaxRank>> OutputShape {InputShape.GetData()};
+				OutputShape[Axis] = (uint32) Splits[Idx];
+				OutputTensors[Idx]->SetShape(NNE::FTensorShape::Make(OutputShape));
+			}
+
+			return 0;
+		};
+
+		virtual bool Initialize(TConstArrayView<NNE::FTensorDesc> InputTensorDescs, TConstArrayView<NNE::FTensorDesc> OutputTensorDescs, const NNE::FAttributeMap& Attributes) override
+		{
+			
+			check(InputTensorDescs.Num() >= 1 && InputTensorDescs.Num() <= 2);
+
+			const int32 InputRank = InputTensorDescs[0].GetShape().Rank();
+			Axis = Attributes.GetValueOrDefault<int32>(TEXT("axis"), Axis);
+			if (Axis > InputRank || Axis  < -InputRank)
+			{
+				UE_LOG(LogNNE, Warning, TEXT("Split's attribute 'Axis' should be in the range [-r,r] with r being the rank of the input (name: %s) however got %d while rank is %d."), *InputTensorDescs[0].GetName(), Axis, InputRank);
+				return false;
+			}
+			// Canonicalize Axis
+			if (Axis < 0)
+			{
+				Axis += InputRank;
+			}
+
+			if constexpr (Version >= 18)
+			{
+				check( ((bool) Attributes.GetAttributeValue(TEXT("num_outputs"))) != (InputTensorDescs.Num() == 2) );
+			}
+
+			if(Version >= 18 && Attributes.GetAttributeValue(TEXT("num_outputs")))
+			{
+				const int32 NumOutputs = Attributes.GetValue<int32>(TEXT("num_outputs"));
+				if(OutputTensorDescs.Num() != NumOutputs)
+				{
+					UE_LOG(LogNNE, Warning, TEXT("Split's attribute 'num_outputs' doesn't match number of output tensors. Value: %d. Number of output tensors: %d."), NumOutputs, OutputTensorDescs.Num());
+					return false;
+				}
+			}
+
+			if(Version < 13 && Attributes.GetAttributeValue(TEXT("split")))
+			{
+				for(const int32& Length : Attributes.GetValue<TArray<int32>>(TEXT("split")))
+				{
+					Splits.Add((int64) Length);
+				}
+			}
+
+			return true;
+		}
+
+
+		virtual void Dispatch(FRDGBuilder& GraphBuilder, TConstArrayView<FTensorRDGRef> InputTensors, TConstArrayView<FTensorRDGRef> OutputTensors) override
+		{
+			using namespace UE::NNEHlslShaders::Internal;
+
+			check(InputTensors.Num() >= 1 && InputTensors.Num() <= 2);
+			check(OutputTensors.Num() == Splits.Num());
+
+			// Set shader parameters
+			FSplitCS::FParameters* Params = GraphBuilder.AllocParameters<FSplitCS::FParameters>();
+			const FTensorRDG& Input = *InputTensors[0];
+			const FRDGBufferSRVRef InputSRV = GraphBuilder.CreateSRV(FRDGBufferSRVDesc(Input.GetBuffer(), PF_R32_FLOAT));
+			Params->Input = InputSRV;
+			FillTensorStrideShaderParameters(Input, Params->InputTensorInfo, /*Idx*/ 0);
+			FillTensorSizeShaderParameters(Input, Params->InputTensorInfo, /*Idx*/ 1);
+			for(int32 TensorIdx = 0; TensorIdx < OutputTensors.Num(); TensorIdx++)
+			{
+				const FTensorRDG& Output = *OutputTensors[TensorIdx];
+				const FRDGBufferUAVRef OutputUAV = GraphBuilder.CreateUAV(FRDGBufferUAVDesc(Output.GetBuffer(), PF_R32_FLOAT));
+				Params->Output[TensorIdx] = OutputUAV;
+				
+				FTensorInfoParamArraySpan TensorInfoParamArraySpan {
+					/*ArrayAtFunction*/ 		[&Params](uint32 Idx) -> FUintVector4& { return Params->OutputTensorInfo[Idx]; },
+					/*Offset*/ 					TensorIdx * FSplitConstants::MAX_NUM_DIMENSIONS
+				};
+
+				FillTensorStrideShaderParameters(Output, TensorInfoParamArraySpan, /*Idx*/ 0);
+				FillTensorSizeShaderParameters(Output, TensorInfoParamArraySpan, /*Idx*/ 1);
+			}
+			Params->Num = Input.GetVolume();
+
+			const FIntVector ThreadGroupCount = ComputeElementWiseThreadGroups(Input.GetVolume(), FSplitConstants::NUM_GROUP_THREADS);
+			Params->ThreadCountX = ThreadGroupCount.X * FSplitConstants::NUM_GROUP_THREADS;
+
+			FSplitCS::FPermutationDomain PermutationVector;
+			PermutationVector.Set<FSplitCS::FSplitRank>((int32) InputTensors[0]->GetShape().Rank());
+			PermutationVector.Set<FSplitCS::FSplitAxis>(Axis);
+			PermutationVector.Set<FSplitCS::FSplitNumSplits>(Splits.Num());
+			
+
+			TShaderMapRef<FSplitCS> ComputeShader(GetGlobalShaderMap(GMaxRHIFeatureLevel), PermutationVector);
+        
+            RDG_EVENT_SCOPE(GraphBuilder, "NNE.Operator.Hlsl.Split");
+            RDG_GPU_STAT_SCOPE(GraphBuilder, FNNEOperatorSplit);
+
+            FComputeShaderUtils::AddPass(
+                GraphBuilder,
+                RDG_EVENT_NAME("NNE.Operator.Hlsl.Split.Dispatch"),
+                ERDGPassFlags::Compute | ERDGPassFlags::NeverCull,
+                ComputeShader,
+                Params,
+                ThreadGroupCount);
+			
+		}
+	};
+
+	template<int32 Version>
+	bool ValidateSplitOperator(const NNE::FAttributeMap& AttributeMap, TConstArrayView<ENNETensorDataType> InputTypes, TConstArrayView<NNE::FSymbolicTensorShape> InputShapes)
+	{
+		bool bIsValid = true;
+
+		FAttributeValidator AttributeValidator;
+		AttributeValidator.AddOptional(TEXT("axis"), ENNEAttributeDataType::Int32);
+		if constexpr (Version < 13)
+		{
+			AttributeValidator.AddOptional(TEXT("split"), ENNEAttributeDataType::Int32Array);
+		}
+		if constexpr (Version >= 18)
+		{
+			AttributeValidator.AddOptional(TEXT("num_outputs"), ENNEAttributeDataType::Int32);
+		}
+        bIsValid &= AttributeValidator.Validate(AttributeMap);
+
+		FInputValidator InputValidator;
+		InputValidator.SetTemplateCount(2);
+		InputValidator.AddSupportedType(ENNETensorDataType::Float);
+		InputValidator.AddSupportedType(ENNETensorDataType::Int64, /* TemplateIdx = */ 1 );
+		InputValidator.AddRequired();
+		if constexpr (Version >= 13)
+		{
+			InputValidator.AddOptional(/* TemplateIdx = */ 1);
+		}
+		bIsValid &= InputValidator.Validate(InputTypes);
+
+		if constexpr (Version >= 18)
+		{
+			bIsValid &= ((bool) AttributeMap.GetAttributeValue(TEXT("num_outputs"))) != (InputTypes.Num() == 2);
+		}
+
+		return bIsValid;
+	}
+
+	template<int32 Version>
+	FOperatorHlsl* CreateSplitOperator()
+	{
+		return new FSplitOperator<Version>();
+	}
+
+
+	bool RegisterSplitOperator(FOperatorRegistryHlsl& Registry)
+	{
+		// Note: support of a particular version is partial with respect to tensor data types (only the most typical ones are usually supported).
+		#define REG_SPLIT_VER(Version) \
+		Registry.OpAdd({{TEXT("Split"), TEXT("Onnx")}, Version}, CreateSplitOperator<Version>, ValidateSplitOperator<Version>);
+
+		REG_SPLIT_VER(2);
+		REG_SPLIT_VER(11);
+		REG_SPLIT_VER(13);
+		REG_SPLIT_VER(18);
+		return true;
+	}
+} // UE::NNERuntimeRDG::Private::Hlsl
diff --git a/Engine/Plugins/Experimental/NNERuntimeRDG/Source/NNERuntimeRDG/Private/Hlsl/NNERuntimeRDGSplit.h b/Engine/Plugins/Experimental/NNERuntimeRDG/Source/NNERuntimeRDG/Private/Hlsl/NNERuntimeRDGSplit.h
new file mode 100644
--- /dev/null
+++ b/Engine/Plugins/Experimental/NNERuntimeRDG/Source/NNERuntimeRDG/Private/Hlsl/NNERuntimeRDGSplit.h
@@ -0,0 +1,11 @@
+// Copyright Epic Games, Inc. All Rights Reserved.
+
+#pragma once
+
+#include "NNERuntimeRDGBase.h"
+#include "NNERuntimeRDGHlslOp.h"
+
+namespace UE::NNERuntimeRDG::Private::Hlsl
+{
+	bool RegisterSplitOperator(FOperatorRegistryHlsl& Registry);
+} // UE::NNERuntimeRDG::Private::Hlsl
diff --git a/Engine/Plugins/Experimental/NNERuntimeRDG/Source/NNERuntimeRDG/Private/NNERuntimeRDGBase.cpp b/Engine/Plugins/Experimental/NNERuntimeRDG/Source/NNERuntimeRDG/Private/NNERuntimeRDGBase.cpp
--- a/Engine/Plugins/Experimental/NNERuntimeRDG/Source/NNERuntimeRDG/Private/NNERuntimeRDGBase.cpp
+++ b/Engine/Plugins/Experimental/NNERuntimeRDG/Source/NNERuntimeRDG/Private/NNERuntimeRDGBase.cpp
@@ -30,7 +30,7 @@ bool AlwaysValidValidationFunction(
 FInputValidator::FInputValidator() : 
 	NumRequiredInput(0), NumOptionalInput(0)
 {
-	TemplateTypes.SetNum(1);
+	SetTemplateCount(1);
 }
 
 bool FInputValidator::Validate(TConstArrayView<ENNETensorDataType> InputTypes)
@@ -56,7 +56,7 @@ bool FInputValidator::Validate(TConstArrayView<ENNETensorDataType> InputTypes)
 		const int32 TemplateIdx = InputTemplateIndices[Idx];
 		
 		check(TemplateIdx < TemplateTypes.Num());
-		if (INDEX_NONE == TemplateTypes[TemplateIdx].Find(InputTypes[Idx]))
+		if (!TemplateTypes[TemplateIdx].Contains(InputTypes[Idx]))
 		{
 			FString TargetType = LogHelper::GetTensorDataTypeName(InputTypes[Idx]);
 			UE_LOG(LogNNE, Warning, TEXT("Input at index '%d' (from template T%d) is of type '%s' witch is not supported for that input."), Idx, TemplateIdx, *TargetType);
@@ -65,11 +65,11 @@ bool FInputValidator::Validate(TConstArrayView<ENNETensorDataType> InputTypes)
 	}
 	return bAreInputValid;
 }
-void FInputValidator::SetTemplateCount(int TemplateCount)
+void FInputValidator::SetTemplateCount(int32 TemplateCount)
 {
-	TemplateTypes.SetNum(TemplateCount);
+	TemplateTypes.Init(TSet<ENNETensorDataType>{}, TemplateCount);
 }
-void FInputValidator::AddSupportedType(ENNETensorDataType Type, int TemplateIdx)
+void FInputValidator::AddSupportedType(ENNETensorDataType Type, int32 TemplateIdx)
 {
 	check(TemplateTypes.Num() > TemplateIdx);
 	TemplateTypes[TemplateIdx].Add(Type);
diff --git a/Engine/Plugins/Experimental/NNERuntimeRDG/Source/NNERuntimeRDG/Private/NNERuntimeRDGBase.h b/Engine/Plugins/Experimental/NNERuntimeRDG/Source/NNERuntimeRDG/Private/NNERuntimeRDGBase.h
--- a/Engine/Plugins/Experimental/NNERuntimeRDG/Source/NNERuntimeRDG/Private/NNERuntimeRDGBase.h
+++ b/Engine/Plugins/Experimental/NNERuntimeRDG/Source/NNERuntimeRDG/Private/NNERuntimeRDGBase.h
@@ -87,8 +87,8 @@ public:
 	bool Validate(TConstArrayView<ENNETensorDataType> InputTypes);
 
 private:
-	TArray<TArray<ENNETensorDataType>> TemplateTypes;
-	TArray<int32> InputTemplateIndices;
+	TArray<TSet<ENNETensorDataType>, TInlineAllocator<16>> TemplateTypes;
+	FIntArray InputTemplateIndices;
 	int32 NumRequiredInput;
 	int32 NumOptionalInput;
 };
diff --git a/Engine/Plugins/Experimental/NNERuntimeRDG/Source/NNERuntimeRDG/Private/NNERuntimeRDGHlsl.cpp b/Engine/Plugins/Experimental/NNERuntimeRDG/Source/NNERuntimeRDG/Private/NNERuntimeRDGHlsl.cpp
--- a/Engine/Plugins/Experimental/NNERuntimeRDG/Source/NNERuntimeRDG/Private/NNERuntimeRDGHlsl.cpp
+++ b/Engine/Plugins/Experimental/NNERuntimeRDG/Source/NNERuntimeRDG/Private/NNERuntimeRDGHlsl.cpp
@@ -37,6 +37,7 @@
 #include "Hlsl/NNERuntimeRDGShape.h"
 #include "Hlsl/NNERuntimeRDGSize.h"
 #include "Hlsl/NNERuntimeRDGSlice.h"
+#include "Hlsl/NNERuntimeRDGSplit.h"
 #include "Hlsl/NNERuntimeRDGSoftmax.h"
 #include "Hlsl/NNERuntimeRDGSqueeze.h"
 #include "Hlsl/NNERuntimeRDGTranspose.h"
@@ -76,6 +77,7 @@ bool UNNERuntimeRDGHlslImpl::Init()
 	RegisterShapeOperator(*Registry);
 	RegisterSizeOperator(*Registry);
 	RegisterSliceOperator(*Registry);
+	RegisterSplitOperator(*Registry);
 	RegisterSoftmaxOperator(*Registry);
 	RegisterSqueezeOperator(*Registry);
 	RegisterTransposeOperator(*Registry);
diff --git a/Engine/Plugins/Experimental/NNERuntimeRDG/Source/NNERuntimeRDG/Private/NNERuntimeRDGHlslHelper.cpp b/Engine/Plugins/Experimental/NNERuntimeRDG/Source/NNERuntimeRDG/Private/NNERuntimeRDGHlslHelper.cpp
--- a/Engine/Plugins/Experimental/NNERuntimeRDG/Source/NNERuntimeRDG/Private/NNERuntimeRDGHlslHelper.cpp
+++ b/Engine/Plugins/Experimental/NNERuntimeRDG/Source/NNERuntimeRDG/Private/NNERuntimeRDGHlslHelper.cpp
@@ -7,16 +7,20 @@
 
 namespace UE::NNERuntimeRDG::Private::Hlsl
 {
-	void FillTensorSizeShaderParameters(const NNE::Internal::FTensor& Tensor, TStaticArray<FUintVector4, NXRT_TENSORSTRIDEINFO_MAX_NUM_DIMENSIONS, 16U>& OutShaderParam, int32 Idx)
+	template<typename TTensorInfoParam>
+	void FillTensorSizeShaderParameters(const NNE::Internal::FTensor& Tensor, TTensorInfoParam& OutTensorInfoParam, int32 Idx)
 	{
 		static_assert(NNE::FTensorShape::MaxRank <= NXRT_TENSORSTRIDEINFO_MAX_NUM_DIMENSIONS);
 		for (int32 i = 0; i < Tensor.GetShape().Rank(); ++i)
 		{
-			OutShaderParam[i][Idx] = Tensor.GetShape().GetData()[i];
+			OutTensorInfoParam[i][Idx] = Tensor.GetShape().GetData()[i];
 		}
 	}
-
-	void FillTensorStrideShaderParameters(const NNE::Internal::FTensor& Tensor, TStaticArray<FUintVector4, NXRT_TENSORSTRIDEINFO_MAX_NUM_DIMENSIONS, 16U>& OutShaderParam, int32 Idx, int32 TargetNumdimensionForBroadcast)
+	template void FillTensorSizeShaderParameters(const NNE::Internal::FTensor&, FTensorInfoParam&, int32);
+	template void FillTensorSizeShaderParameters(const NNE::Internal::FTensor&, FTensorInfoParamArraySpan&, int32);
+	
+	template<typename TTensorInfoParam>
+	void FillTensorStrideShaderParameters(const NNE::Internal::FTensor& Tensor, TTensorInfoParam& OutTensorInfoParam, int32 Idx, int32 TargetNumdimensionForBroadcast)
 	{
 		if (TargetNumdimensionForBroadcast == -1)
 		{
@@ -30,33 +34,38 @@ namespace UE::NNERuntimeRDG::Private::Hlsl
 		{
 			if (i >= TargetNumdimensionForBroadcast || i < Offset)
 			{
-				OutShaderParam[i][Idx] = 0;
+				OutTensorInfoParam[i][Idx] = 0;
 			}
 			else if (i == TargetNumdimensionForBroadcast - 1)
 			{
-				OutShaderParam[i][Idx] = 1;
+				OutTensorInfoParam[i][Idx] = 1;
 			}
 			else
 			{
-				OutShaderParam[i][Idx] = OutShaderParam[i + 1][Idx] * Tensor.GetShape().GetData()[i + 1 - Offset];
+				OutTensorInfoParam[i][Idx] = OutTensorInfoParam[i + 1][Idx] * Tensor.GetShape().GetData()[i + 1 - Offset];
 			}
 		}
 	}
+	template void FillTensorStrideShaderParameters(const NNE::Internal::FTensor&, FTensorInfoParam&, int32, int32);
+	template void FillTensorStrideShaderParameters(const NNE::Internal::FTensor&, FTensorInfoParamArraySpan&, int32, int32);
 
-	void FillTensorStrideForBroadcastShaderParameters(const NNE::Internal::FTensor& Tensor, int32 OutputNumdimension, TStaticArray<FUintVector4, NXRT_TENSORSTRIDEINFO_MAX_NUM_DIMENSIONS, 16U>& OutShaderParam, int32 Idx)
+	template<typename TTensorInfoParam>
+	void FillTensorStrideForBroadcastShaderParameters(const NNE::Internal::FTensor& Tensor, int32 OutputNumdimension, TTensorInfoParam& OutTensorInfoParam, int32 Idx)
 	{
 		checkf(OutputNumdimension >= Tensor.GetShape().Rank(), TEXT("Can't broadcast tensor from rank %d to rank %d, should be inferior or equal."), Tensor.GetShape().Rank(), OutputNumdimension);
-		FillTensorStrideShaderParameters(Tensor, OutShaderParam, Idx, OutputNumdimension);
+		FillTensorStrideShaderParameters(Tensor, OutTensorInfoParam, Idx, OutputNumdimension);
 		int32 Offset = OutputNumdimension - Tensor.GetShape().Rank();
 		for (int32 i = Offset; i < OutputNumdimension; ++i)
 		{
 			// the stride for broadcast dimension is kept as 0
 			if (Tensor.GetShape().GetData()[i - Offset] == 1)
 			{
-				OutShaderParam[i][Idx] = 0;
+				OutTensorInfoParam[i][Idx] = 0;
 			}
 		}
 	}
+	template void FillTensorStrideForBroadcastShaderParameters(const NNE::Internal::FTensor&, int32, FTensorInfoParam&, int32);
+	template void FillTensorStrideForBroadcastShaderParameters(const NNE::Internal::FTensor&, int32, FTensorInfoParamArraySpan&, int32);
 
 	FIntVector ComputeElementWiseThreadGroups(uint32 ElementCount, uint32 GroupSizeX)
 	{
diff --git a/Engine/Plugins/Experimental/NNERuntimeRDG/Source/NNERuntimeRDG/Private/NNERuntimeRDGHlslHelper.h b/Engine/Plugins/Experimental/NNERuntimeRDG/Source/NNERuntimeRDG/Private/NNERuntimeRDGHlslHelper.h
--- a/Engine/Plugins/Experimental/NNERuntimeRDG/Source/NNERuntimeRDG/Private/NNERuntimeRDGHlslHelper.h
+++ b/Engine/Plugins/Experimental/NNERuntimeRDG/Source/NNERuntimeRDG/Private/NNERuntimeRDGHlslHelper.h
@@ -3,6 +3,7 @@
 #pragma once
 
 #include "Containers/StaticArray.h"
+#include "Templates/Function.h"
 
 namespace UE::NNE::Internal { class FTensor; }
 
@@ -10,9 +11,26 @@ namespace UE::NNE::Internal { class FTensor; }
 
 namespace UE::NNERuntimeRDG::Private::Hlsl
 {
+	using FTensorInfoParam = TStaticArray<FUintVector4, NXRT_TENSORSTRIDEINFO_MAX_NUM_DIMENSIONS, 16U>;
 
-	void FillTensorSizeShaderParameters(const NNE::Internal::FTensor& Tensor, TStaticArray<FUintVector4, NXRT_TENSORSTRIDEINFO_MAX_NUM_DIMENSIONS, 16U>& OutShaderParam, int32 Idx);
-	void FillTensorStrideShaderParameters(const NNE::Internal::FTensor& Tensor, TStaticArray<FUintVector4, NXRT_TENSORSTRIDEINFO_MAX_NUM_DIMENSIONS, 16U>& OutShaderParam, int32 Idx, int32 TargetNumdimensionForBroadcast = -1);
-	void FillTensorStrideForBroadcastShaderParameters(const NNE::Internal::FTensor& Tensor, int32 OutputNumdimension, TStaticArray<FUintVector4, NXRT_TENSORSTRIDEINFO_MAX_NUM_DIMENSIONS, 16U>& OutShaderParam, int32 Idx);
+	struct FTensorInfoParamArraySpan 
+	{
+		TFunction<FUintVector4& (uint32)> ArrayAtFunction;
+		uint32 Offset = 0;
+		static constexpr uint32 Length = NXRT_TENSORSTRIDEINFO_MAX_NUM_DIMENSIONS;
+		FUintVector4& operator[](uint32 Idx) 
+		{
+			check(Idx < Length);
+			check(ArrayAtFunction);
+			return ArrayAtFunction(Offset + Idx); 
+		}
+	};
+
+	template<typename TTensorInfoParam>
+	void FillTensorSizeShaderParameters(const NNE::Internal::FTensor& Tensor, TTensorInfoParam& OutTensorInfoParam, int32 Idx);
+	template<typename TTensorInfoParam>
+	void FillTensorStrideShaderParameters(const NNE::Internal::FTensor& Tensor, TTensorInfoParam& OutTensorInfoParam, int32 Idx, int32 TargetNumdimensionForBroadcast = -1);
+	template<typename TTensorInfoParam>
+	void FillTensorStrideForBroadcastShaderParameters(const NNE::Internal::FTensor& Tensor, int32 OutputNumdimension, TTensorInfoParam& OutTensorInfoParam, int32 Idx);
 	FIntVector ComputeElementWiseThreadGroups(uint32 ElementCount, uint32 GroupSizeX);
 } // namespace UE::NNERuntimeRDG::Private::Hlsl
\ No newline at end of file
