[NNE] HLSL Add a console command to log onnx operator coverage (name, domain, versions) + fix reduce and global pool ops registration

Problems:
* We didn't had a way to know what operator/opsets are supported without looking at the code or importing related models.
* Some hlsl operators where not registered with there matching opset.

Solution:
* Add a console command to log the support matrix in csv format.
* Fix the operators registration

#tests cvar in win64 editor, all nne tests win64 editor and standalone

#rnx

#virtualized

diff --git a/Engine/Plugins/Experimental/NNERuntimeRDG/Source/NNERuntimeRDG/Private/Hlsl/NNERuntimeRDGGlobalPool.cpp b/Engine/Plugins/Experimental/NNERuntimeRDG/Source/NNERuntimeRDG/Private/Hlsl/NNERuntimeRDGGlobalPool.cpp
--- a/Engine/Plugins/Experimental/NNERuntimeRDG/Source/NNERuntimeRDG/Private/Hlsl/NNERuntimeRDGGlobalPool.cpp
+++ b/Engine/Plugins/Experimental/NNERuntimeRDG/Source/NNERuntimeRDG/Private/Hlsl/NNERuntimeRDGGlobalPool.cpp
@@ -121,8 +121,8 @@ namespace UE::NNERuntimeRDG::Private::Hlsl
 
 	bool RegisterGlobalPoolOperators(FOperatorRegistryHlsl& Registry)
 	{
-		Registry.OpAdd({{TEXT("GlobalAveragePool"), TEXT("Onnx")}}, CreateGlobalPoolOperator<UE::NNEHlslShaders::Internal::EReduceOperatorType::Average>, ValidateGlobalPoolOperator);
-		Registry.OpAdd({ {TEXT("GlobalMaxPool"), TEXT("Onnx")} }, CreateGlobalPoolOperator<UE::NNEHlslShaders::Internal::EReduceOperatorType::Max>, ValidateGlobalPoolOperator);
+		Registry.OpAdd({{TEXT("GlobalAveragePool"), TEXT("Onnx")}, 1}, CreateGlobalPoolOperator<UE::NNEHlslShaders::Internal::EReduceOperatorType::Average>, ValidateGlobalPoolOperator);
+		Registry.OpAdd({{TEXT("GlobalMaxPool"), TEXT("Onnx")}, 1}, CreateGlobalPoolOperator<UE::NNEHlslShaders::Internal::EReduceOperatorType::Max>, ValidateGlobalPoolOperator);
 		return true;
 	}
 } // UE::NNERuntimeRDG::Private::Hlsl
diff --git a/Engine/Plugins/Experimental/NNERuntimeRDG/Source/NNERuntimeRDG/Private/Hlsl/NNERuntimeRDGReduce.cpp b/Engine/Plugins/Experimental/NNERuntimeRDG/Source/NNERuntimeRDG/Private/Hlsl/NNERuntimeRDGReduce.cpp
--- a/Engine/Plugins/Experimental/NNERuntimeRDG/Source/NNERuntimeRDG/Private/Hlsl/NNERuntimeRDGReduce.cpp
+++ b/Engine/Plugins/Experimental/NNERuntimeRDG/Source/NNERuntimeRDG/Private/Hlsl/NNERuntimeRDGReduce.cpp
@@ -168,17 +168,27 @@ namespace UE::NNERuntimeRDG::Private::Hlsl
 
 	bool RegisterReduceOperators(FOperatorRegistryHlsl& Registry)
 	{
-		Registry.OpAdd({ {TEXT("ReduceL1"),       TEXT("Onnx")}}, CreateReduceOperator<UE::NNEHlslShaders::Internal::EReduceOperatorType::L1>, ValidateReduceOperator);
-		Registry.OpAdd({ {TEXT("ReduceL2"),       TEXT("Onnx")}}, CreateReduceOperator<UE::NNEHlslShaders::Internal::EReduceOperatorType::L2>, ValidateReduceOperator);
 		//ReduceLogSum not yet supported as multi axis case require to apply all reduction first then log.
-		Registry.OpAdd({{TEXT("ReduceLogSumExp"), TEXT("Onnx")}}, CreateReduceOperator<UE::NNEHlslShaders::Internal::EReduceOperatorType::LogSumExp>, ValidateReduceOperator);
-		Registry.OpAdd({{TEXT("ReduceMax"),       TEXT("Onnx")}}, CreateReduceOperator<UE::NNEHlslShaders::Internal::EReduceOperatorType::Max>, ValidateReduceOperator);
-		Registry.OpAdd({{TEXT("ReduceMean"),      TEXT("Onnx")}}, CreateReduceOperator<UE::NNEHlslShaders::Internal::EReduceOperatorType::Average>, ValidateReduceOperator);
-		Registry.OpAdd({{TEXT("ReduceMin"),       TEXT("Onnx")}}, CreateReduceOperator<UE::NNEHlslShaders::Internal::EReduceOperatorType::Min>, ValidateReduceOperator);
-		Registry.OpAdd({{TEXT("ReduceProd"),      TEXT("Onnx")}}, CreateReduceOperator<UE::NNEHlslShaders::Internal::EReduceOperatorType::Prod>, ValidateReduceOperator);
-		Registry.OpAdd({{TEXT("ReduceSum"),       TEXT("Onnx")}}, CreateReduceOperator<UE::NNEHlslShaders::Internal::EReduceOperatorType::Sum>, ValidateReduceOperator);
 		//ReduceSumSquare not yet supported as multi axis case require to apply square on whole tensor first then sum.
 
+		Registry.OpAdd({ {TEXT("ReduceL1"),        TEXT("Onnx")},  1  }, CreateReduceOperator<UE::NNEHlslShaders::Internal::EReduceOperatorType::L1>, ValidateReduceOperator);
+		Registry.OpAdd({ {TEXT("ReduceL2"),        TEXT("Onnx")},  1  }, CreateReduceOperator<UE::NNEHlslShaders::Internal::EReduceOperatorType::L2>, ValidateReduceOperator);
+		Registry.OpAdd({ {TEXT("ReduceLogSumExp"), TEXT("Onnx")},  1  }, CreateReduceOperator<UE::NNEHlslShaders::Internal::EReduceOperatorType::LogSumExp>, ValidateReduceOperator);
+		Registry.OpAdd({ {TEXT("ReduceMax"),       TEXT("Onnx")},  1  }, CreateReduceOperator<UE::NNEHlslShaders::Internal::EReduceOperatorType::Max>, ValidateReduceOperator);
+		Registry.OpAdd({ {TEXT("ReduceMean"),      TEXT("Onnx")},  1  }, CreateReduceOperator<UE::NNEHlslShaders::Internal::EReduceOperatorType::Average>, ValidateReduceOperator);
+		Registry.OpAdd({ {TEXT("ReduceMin"),       TEXT("Onnx")},  1  }, CreateReduceOperator<UE::NNEHlslShaders::Internal::EReduceOperatorType::Min>, ValidateReduceOperator);
+		Registry.OpAdd({ {TEXT("ReduceProd"),      TEXT("Onnx")},  1  }, CreateReduceOperator<UE::NNEHlslShaders::Internal::EReduceOperatorType::Prod>, ValidateReduceOperator);
+		Registry.OpAdd({ {TEXT("ReduceSum"),       TEXT("Onnx")},  1  }, CreateReduceOperator<UE::NNEHlslShaders::Internal::EReduceOperatorType::Sum>, ValidateReduceOperator);
+
+		Registry.OpAdd({ {TEXT("ReduceL1"),        TEXT("Onnx")}, 11 }, CreateReduceOperator<UE::NNEHlslShaders::Internal::EReduceOperatorType::L1>, ValidateReduceOperator);
+		Registry.OpAdd({ {TEXT("ReduceL2"),        TEXT("Onnx")}, 11 }, CreateReduceOperator<UE::NNEHlslShaders::Internal::EReduceOperatorType::L2>, ValidateReduceOperator);
+		Registry.OpAdd({ {TEXT("ReduceLogSumExp"), TEXT("Onnx")}, 11 }, CreateReduceOperator<UE::NNEHlslShaders::Internal::EReduceOperatorType::LogSumExp>, ValidateReduceOperator);
+		Registry.OpAdd({ {TEXT("ReduceMax"),       TEXT("Onnx")}, 11 }, CreateReduceOperator<UE::NNEHlslShaders::Internal::EReduceOperatorType::Max>, ValidateReduceOperator);
+		Registry.OpAdd({ {TEXT("ReduceMean"),      TEXT("Onnx")}, 11 }, CreateReduceOperator<UE::NNEHlslShaders::Internal::EReduceOperatorType::Average>, ValidateReduceOperator);
+		Registry.OpAdd({ {TEXT("ReduceMin"),       TEXT("Onnx")}, 11 }, CreateReduceOperator<UE::NNEHlslShaders::Internal::EReduceOperatorType::Min>, ValidateReduceOperator);
+		Registry.OpAdd({ {TEXT("ReduceProd"),      TEXT("Onnx")}, 11 }, CreateReduceOperator<UE::NNEHlslShaders::Internal::EReduceOperatorType::Prod>, ValidateReduceOperator);
+		Registry.OpAdd({ {TEXT("ReduceSum"),       TEXT("Onnx")}, 11 }, CreateReduceOperator<UE::NNEHlslShaders::Internal::EReduceOperatorType::Sum>, ValidateReduceOperator);
+
 		return true;
 	}
 } // UE::NNERuntimeRDG::Private::Hlsl
diff --git a/Engine/Plugins/Experimental/NNERuntimeRDG/Source/NNERuntimeRDG/Private/NNERuntimeRDGBase.h b/Engine/Plugins/Experimental/NNERuntimeRDG/Source/NNERuntimeRDG/Private/NNERuntimeRDGBase.h
--- a/Engine/Plugins/Experimental/NNERuntimeRDG/Source/NNERuntimeRDG/Private/NNERuntimeRDGBase.h
+++ b/Engine/Plugins/Experimental/NNERuntimeRDG/Source/NNERuntimeRDG/Private/NNERuntimeRDGBase.h
@@ -233,6 +233,78 @@ public:
 		return true;
 	}
 
+	FString ListAllRegisteredOperators() const
+	{
+		class FLoggableOperatorDesc
+		{
+			FString OpDomainAndName;
+			FString SupportedOpsets;
+
+		public:
+
+			FLoggableOperatorDesc(const FOperatorDescUnversioned& InOpDesc, TArrayView<TOptional<TOperatorVersionType>> InOpVersions)
+			{
+				OpDomainAndName = FString::Printf(TEXT("%s,%s"), *InOpDesc.DomainName, *InOpDesc.OpName);
+
+				TArray<TOptional<TOperatorVersionType>, TInlineAllocator<8>> OpVersions(InOpVersions);
+
+				OpVersions.Sort([](const TOptional<TOperatorVersionType>& lhs, const TOptional<TOperatorVersionType>& rhs)
+					{
+						if (!lhs.IsSet())
+						{
+							return true;
+						}
+						else if (!rhs.IsSet())
+						{
+							return false;
+						}
+						else
+						{
+							return *lhs < *rhs;
+						}
+					});
+
+				for (const TOptional<TOperatorVersionType>& OpVersion : OpVersions)
+				{
+					SupportedOpsets += (OpVersion.IsSet() ? FString::Printf(TEXT(",%d"), *OpVersion) : TEXT(",Unversioned"));
+				}
+			}
+
+			const FString& GetDomainAndName() const
+			{
+				return OpDomainAndName;
+			}
+			const FString& GetSupportedOpset() const
+			{
+				return SupportedOpsets;
+			}
+		};
+
+		TArray<FLoggableOperatorDesc> LoggableOperatorDesc;
+		TArray<TOptional<TOperatorVersionType>> Versions;
+
+		for (const TPair<FOperatorDescUnversioned, TOperatorVersionToFunctionsMap>& Pair : Operators)
+		{
+			Pair.Value.GenerateKeyArray(Versions);
+			LoggableOperatorDesc.Emplace(Pair.Key, Versions);
+		}
+
+		LoggableOperatorDesc.Sort([](const FLoggableOperatorDesc& lhs, const FLoggableOperatorDesc& rhs) 
+		{
+			return lhs.GetDomainAndName() < rhs.GetDomainAndName();
+		});
+
+		FString OpListWithOpset;
+		for (const FLoggableOperatorDesc& LoggableOp : LoggableOperatorDesc)
+		{
+			OpListWithOpset += LoggableOp.GetDomainAndName();
+			OpListWithOpset += LoggableOp.GetSupportedOpset();
+			OpListWithOpset += TEXT("\n");
+		}
+
+		return OpListWithOpset;
+	}
+
 private:
 
 	struct FOperatorFunctions
diff --git a/Engine/Plugins/Experimental/NNERuntimeRDG/Source/NNERuntimeRDG/Private/NNERuntimeRDGHlsl.cpp b/Engine/Plugins/Experimental/NNERuntimeRDG/Source/NNERuntimeRDG/Private/NNERuntimeRDGHlsl.cpp
--- a/Engine/Plugins/Experimental/NNERuntimeRDG/Source/NNERuntimeRDG/Private/NNERuntimeRDGHlsl.cpp
+++ b/Engine/Plugins/Experimental/NNERuntimeRDG/Source/NNERuntimeRDG/Private/NNERuntimeRDGHlsl.cpp
@@ -15,6 +15,7 @@
 #ifdef NNE_UTILITIES_AVAILABLE
 #include "NNEUtilitiesModelOptimizer.h"
 #endif // NNE_UTILITIES_AVAILABLE
+#include "HAL/IConsoleManager.h"
 #include "Hlsl/NNERuntimeRDGBatchNormalization.h"
 #include "Hlsl/NNERuntimeRDGCast.h"
 #include "Hlsl/NNERuntimeRDGConv.h"
@@ -88,6 +89,22 @@ bool UNNERuntimeRDGHlslImpl::Init()
 	return true;
 }
 
+namespace ConsoleCommands
+{
+	static FAutoConsoleCommand GetAutomationRuntimeFilterCommand(
+		TEXT("nne.hlsl.getoperatorsupportmatrix"), TEXT("Get the NNERuntimeRDGHlsl operators support matrix in term of ONNX."),
+		FConsoleCommandWithArgsDelegate::CreateStatic(
+			[](const TArray< FString >& Args)
+			{
+				FOperatorRegistryHlsl* Registry = FOperatorRegistryHlsl::Get();
+				check(Registry != nullptr);
+				FString SupportMatrix = Registry->ListAllRegisteredOperators();
+				UE_LOG(LogNNE, Display, TEXT("HLSL Operators support matrix: \n%s"), *SupportMatrix);
+			}
+		)
+	);
+} // ConsoleCommands
+
 UNNERuntimeRDGHlslImpl::ECanCreateModelDataStatus UNNERuntimeRDGHlslImpl::CanCreateModelData(const FString& FileType, TConstArrayView<uint8> FileData, const TMap<FString, TConstArrayView<uint8>>& AdditionalFileData, const FGuid& FileId, const ITargetPlatform* TargetPlatform) const
 {
 #ifdef NNE_UTILITIES_AVAILABLE
