[NNE] Hlsl enable negative pad values

Problems:
* Hlsl pad doesn't support negative pad values yet
* RDG Dml casts pad values to uint without checking if the values are negative

Solutions:
* Add support for negative pad values on the Hlsl side
* Add check for negative values on the Dml side

Tests:
* NNEQA iteration tests
* Standalone (win64)

#rnx

#virtualized

diff --git a/Engine/Plugins/Experimental/NNERuntimeRDG/Source/NNEHlslShaders/Shaders/NNEHlslShadersPad.usf b/Engine/Plugins/Experimental/NNERuntimeRDG/Source/NNEHlslShaders/Shaders/NNEHlslShadersPad.usf
--- a/Engine/Plugins/Experimental/NNERuntimeRDG/Source/NNEHlslShaders/Shaders/NNEHlslShadersPad.usf
+++ b/Engine/Plugins/Experimental/NNERuntimeRDG/Source/NNEHlslShaders/Shaders/NNEHlslShadersPad.usf
@@ -20,6 +20,31 @@ uint ThreadCountX;
 #define REFLECT_MODE 1
 #define EDGE_MODE 2
 
+#define STATIC_LOOP(Var, From, To) \
+[unroll] \
+for(uint Var = From; Var < To; ++Var)
+
+void GetOutputDimIterator(const uint GlobalIdx, out uint DimIterator[NUM_DIMENSIONS])
+{
+	uint Offset = GlobalIdx;
+	STATIC_LOOP(DimIdx, 0, NUM_DIMENSIONS)
+	{
+		uint Remainder;
+		DivMod(Offset, TensorInfo[DimIdx][OUTPUT_STRIDE], DimIterator[DimIdx], Remainder);
+		Offset = Remainder;
+	}
+}
+
+uint GetInputGlobalIdx(const uint DimIterator[NUM_DIMENSIONS])
+{
+	uint Offset = 0;
+	STATIC_LOOP(DimIdx, 0, NUM_DIMENSIONS)
+	{
+		Offset += DimIterator[DimIdx] * TensorInfo[DimIdx][INPUT_STRIDE];
+	}
+	return Offset;
+}
+
 [numthreads(THREADGROUP_SIZE_X, 1, 1)]
 void Pad(in const uint3 DispatchThreadID : SV_DispatchThreadID)
 {
@@ -27,44 +52,42 @@ void Pad(in const uint3 DispatchThreadID : SV_DispatchThreadID)
 	
 	if (Index < Num)
 	{
-		uint InputIndex = 0;
+		uint OutputDimIterator[NUM_DIMENSIONS];
+		GetOutputDimIterator(Index, OutputDimIterator);
 
-		uint Offset = Index;
-		for (uint dim = 0; dim < NUM_DIMENSIONS; ++dim)
+		uint InputDimIterator[NUM_DIMENSIONS];
+
+		for(uint DimIdx = 0; DimIdx < NUM_DIMENSIONS; ++DimIdx)
 		{
-			uint OutDimIdx, R;
-			DivMod(Offset, TensorInfo[dim][OUTPUT_STRIDE], OutDimIdx, R);
+			const int InputDimSize = (int) TensorInfo[DimIdx][INPUT_SIZE];
 			
-			int PrePad = TensorInfo[dim][PRE_PAD];
-			int InputDimIdx = (int)OutDimIdx - PrePad;
-			int InputDimSize = (int) TensorInfo[dim][INPUT_SIZE];
+			int PrePad = asint(TensorInfo[DimIdx][PRE_PAD]); // Decode PrePad to int (can be negative)
+			int CurDimIterator = (int) OutputDimIterator[DimIdx] - PrePad;
 			
-			if (InputDimIdx < 0 || InputDimIdx >= InputDimSize)
+			if (CurDimIterator < 0 || CurDimIterator >= InputDimSize)
 			{
 				#if MODE == CONSTANT_MODE
-				Output[Index] = Value;
-				return;
+					Output[Index] = Value;
+					return;
 				#elif MODE == REFLECT_MODE
-				if (InputDimIdx >= InputDimSize)
-				{
-					int OverCnt = InputDimIdx - (InputDimSize - 2);
-					InputDimIdx = InputDimSize - OverCnt;
-				}
-				else if (InputDimIdx < 0)
-				{
-					InputDimIdx = -InputDimIdx;
-				}
+					if (CurDimIterator >= InputDimSize)
+					{
+						const int OverCnt = CurDimIterator - (InputDimSize - 2);
+						CurDimIterator = InputDimSize - OverCnt;
+					}
+					else if (CurDimIterator < 0)
+					{
+						CurDimIterator = -CurDimIterator;
+					}
 				#endif
 				//MODE == EDGE_MODE + clamping large reflect indices
-				InputDimIdx = min(InputDimIdx, InputDimSize - 1);
-				InputDimIdx = max(InputDimIdx, 0);
+				CurDimIterator = min(CurDimIterator, InputDimSize - 1);
+				CurDimIterator = max(CurDimIterator, 0);
 			}
-			
-			InputIndex += TensorInfo[dim][INPUT_STRIDE] * InputDimIdx;
-			Offset = R;
+
+			InputDimIterator[DimIdx] = (uint) CurDimIterator;
 		}
 
-		float X = Input[InputIndex];
-		Output[Index] = X;
+		Output[Index] = Input[GetInputGlobalIdx(InputDimIterator)];
 	}
 }
diff --git a/Engine/Plugins/Experimental/NNERuntimeRDG/Source/NNERuntimeRDG/Private/Dml/NNEDmlOperatorPad.cpp b/Engine/Plugins/Experimental/NNERuntimeRDG/Source/NNERuntimeRDG/Private/Dml/NNEDmlOperatorPad.cpp
--- a/Engine/Plugins/Experimental/NNERuntimeRDG/Source/NNERuntimeRDG/Private/Dml/NNEDmlOperatorPad.cpp
+++ b/Engine/Plugins/Experimental/NNERuntimeRDG/Source/NNERuntimeRDG/Private/Dml/NNEDmlOperatorPad.cpp
@@ -89,7 +89,7 @@ public:
 
 		// Read attributes
 		Value = Attributes.GetValueOrDefault<float>(TEXT("value"), 0.0f);
-		Mode = ModeFromString(Attributes.GetValue<FString>(TEXT("mode")));
+		Mode = ModeFromString(Attributes.GetValueOrDefault<FString>(TEXT("mode"), TEXT("constant")));
 
 		if (Inputs.Num() >= 2)
 		{
@@ -122,6 +122,15 @@ public:
 			return -1;
 		}
 
+		for(int32 PadValue : Pads)
+		{
+			if(PadValue < 0)
+			{
+				UE_LOG(LogNNE, Error, TEXT("Pads must be >= 0"));
+				return -1;
+			}
+		}
+
 		TConstArrayView<uint32> InputShape = InputTensors[0]->GetShape().GetData();
 
 		if (InputShape.Num() * 2 != Pads.Num())
diff --git a/Engine/Plugins/Experimental/NNERuntimeRDG/Source/NNERuntimeRDG/Private/Hlsl/NNERuntimeRDGPad.cpp b/Engine/Plugins/Experimental/NNERuntimeRDG/Source/NNERuntimeRDG/Private/Hlsl/NNERuntimeRDGPad.cpp
--- a/Engine/Plugins/Experimental/NNERuntimeRDG/Source/NNERuntimeRDG/Private/Hlsl/NNERuntimeRDGPad.cpp
+++ b/Engine/Plugins/Experimental/NNERuntimeRDG/Source/NNERuntimeRDG/Private/Hlsl/NNERuntimeRDGPad.cpp
@@ -64,7 +64,7 @@ namespace UE::NNERuntimeRDG::Private::Hlsl
 			
 			Pads = Attributes.GetValue<TArray<int32>>(TEXT("pads"));
 			Value = Attributes.GetValueOrDefault<float>(TEXT("value"), 0.0f);
-			FPadCS::LexFromString(Mode, *Attributes.GetValue<FString>(TEXT("mode")));
+			FPadCS::LexFromString(Mode, *Attributes.GetValueOrDefault<FString>(TEXT("mode"), TEXT("constant")));
 
 			if ((2*InputTensorDescs[0].GetShape().Rank()) != Pads.Num())
 			{
@@ -100,7 +100,7 @@ namespace UE::NNERuntimeRDG::Private::Hlsl
 			FillTensorSizeShaderParameters(Input, Params->TensorInfo, 2);
 			for (int32 i = 0; i < Input.GetShape().Rank(); ++i)
 			{
-				Params->TensorInfo[i][3] = Pads[i];//Pre-pad
+				Params->TensorInfo[i][3] = BitCast<uint32>(Pads[i]); // Pre-pad encoded as uint32
 			}
 			Params->Value = Value;
 			Params->Num = Output.GetVolume();
@@ -134,15 +134,6 @@ namespace UE::NNERuntimeRDG::Private::Hlsl
 		AttributeValidator.AddRequired(TEXT("pads"), ENNEAttributeDataType::Int32Array);
 		AttributeValidator.AddOptional(TEXT("value"), ENNEAttributeDataType::Float);
 		bIsValid &= AttributeValidator.Validate(AttributeMap);
-
-		for (int32 Pad : AttributeMap.GetValue<TArray<int32>>(TEXT("pads")))
-		{
-			if (Pad < 0)
-			{
-				UE_LOG(LogNNE, Warning, TEXT("Pad operator does not support negative padding at the moment."));
-				return false;
-			}
-		}
 		
 		FInputValidator InputValidator;
 		InputValidator.AddSupportedType(ENNETensorDataType::Float);
