[NNE] Save runtime format to persisent archive if necessary.

#tests Editor/Game run all tests

#rnx

diff --git a/Engine/Plugins/Experimental/NNERuntimeRDG/Source/NNERuntimeRDG/Private/Dml/NNEDmlModel.cpp b/Engine/Plugins/Experimental/NNERuntimeRDG/Source/NNERuntimeRDG/Private/Dml/NNEDmlModel.cpp
--- a/Engine/Plugins/Experimental/NNERuntimeRDG/Source/NNERuntimeRDG/Private/Dml/NNEDmlModel.cpp
+++ b/Engine/Plugins/Experimental/NNERuntimeRDG/Source/NNERuntimeRDG/Private/Dml/NNEDmlModel.cpp
@@ -44,7 +44,7 @@ FModelInfo* FModelInfo::Get()
 
 FModelInfo::FModelInfo()
 	: Guid((int32)'R', (int32)'D', (int32)'G', (int32)'D')
-	, Version(0x5)
+	, Version(0x6)
 {
 }
 
diff --git a/Engine/Plugins/Experimental/NNERuntimeRDG/Source/NNERuntimeRDG/Private/NNERuntimeRDGBase.h b/Engine/Plugins/Experimental/NNERuntimeRDG/Source/NNERuntimeRDG/Private/NNERuntimeRDGBase.h
--- a/Engine/Plugins/Experimental/NNERuntimeRDG/Source/NNERuntimeRDG/Private/NNERuntimeRDGBase.h
+++ b/Engine/Plugins/Experimental/NNERuntimeRDG/Source/NNERuntimeRDG/Private/NNERuntimeRDGBase.h
@@ -362,6 +362,7 @@ public:
 		}
 
 		FMemoryReader Reader(InputModel.Data);
+		Reader.SetIsPersistent(true);
 
 		Format.Serialize(Reader);
 
diff --git a/Engine/Plugins/Experimental/NNERuntimeRDG/Source/NNERuntimeRDG/Private/NNERuntimeRDGHlsl.cpp b/Engine/Plugins/Experimental/NNERuntimeRDG/Source/NNERuntimeRDG/Private/NNERuntimeRDGHlsl.cpp
--- a/Engine/Plugins/Experimental/NNERuntimeRDG/Source/NNERuntimeRDG/Private/NNERuntimeRDGHlsl.cpp
+++ b/Engine/Plugins/Experimental/NNERuntimeRDG/Source/NNERuntimeRDG/Private/NNERuntimeRDGHlsl.cpp
@@ -49,7 +49,7 @@
 using namespace UE::NNERuntimeRDG::Private::Hlsl;
 
 FGuid UNNERuntimeRDGHlslImpl::GUID = FGuid((int32)'R', (int32)'D', (int32)'G', (int32)'H');
-int32 UNNERuntimeRDGHlslImpl::Version = 0x00000005;
+int32 UNNERuntimeRDGHlslImpl::Version = 0x00000006;
 
 bool UNNERuntimeRDGHlslImpl::Init()
 {
diff --git a/Engine/Plugins/Experimental/NNERuntimeRDG/Source/NNERuntimeRDG/Private/NNERuntimeRDGModel.cpp b/Engine/Plugins/Experimental/NNERuntimeRDG/Source/NNERuntimeRDG/Private/NNERuntimeRDGModel.cpp
--- a/Engine/Plugins/Experimental/NNERuntimeRDG/Source/NNERuntimeRDG/Private/NNERuntimeRDGModel.cpp
+++ b/Engine/Plugins/Experimental/NNERuntimeRDG/Source/NNERuntimeRDG/Private/NNERuntimeRDGModel.cpp
@@ -17,6 +17,7 @@ bool FModelInstanceRDG::LoadModel(TConstArrayView<uint8> ModelData, FNNERuntimeF
 	TConstArrayView<uint8> ModelBuffer = { &(ModelData.GetData()[GuidAndVersionSize]), ModelData.Num() - GuidAndVersionSize };
 
 	FMemoryReaderView Reader(ModelBuffer);
+	Reader.SetIsPersistent(true);
 
 	Format.Serialize(Reader);
 
diff --git a/Engine/Plugins/NNE/NNERuntimeORT/Source/NNEUtilities/Private/NNEUtilitiesModelBuilderNNE.cpp b/Engine/Plugins/NNE/NNERuntimeORT/Source/NNEUtilities/Private/NNEUtilitiesModelBuilderNNE.cpp
--- a/Engine/Plugins/NNE/NNERuntimeORT/Source/NNEUtilities/Private/NNEUtilitiesModelBuilderNNE.cpp
+++ b/Engine/Plugins/NNE/NNERuntimeORT/Source/NNEUtilities/Private/NNEUtilitiesModelBuilderNNE.cpp
@@ -44,8 +44,8 @@ public:
 
 	virtual bool End(TArray<uint8>& Data) override
 	{
-		FMemoryWriter Writer(Data);
-
+		FMemoryWriter Writer(Data, true);
+		
 		Format.Serialize(Writer);
 
 		return !Data.IsEmpty();
