[NNE] Remove NNERuntimeRDGDml.

#tests Editor/Game all tests

#virtualized

diff --git a/Engine/Plugins/Experimental/NNERuntimeRDG/Source/NNERuntimeRDG/NNERuntimeRDG.Build.cs b/Engine/Plugins/Experimental/NNERuntimeRDG/Source/NNERuntimeRDG/NNERuntimeRDG.Build.cs
--- a/Engine/Plugins/Experimental/NNERuntimeRDG/Source/NNERuntimeRDG/NNERuntimeRDG.Build.cs
+++ b/Engine/Plugins/Experimental/NNERuntimeRDG/Source/NNERuntimeRDG/NNERuntimeRDG.Build.cs
@@ -22,26 +22,15 @@ public class NNERuntimeRDG : ModuleRules
 			"RenderCore"
 		});
 
-        PrivateDependencyModuleNames.AddRange(new string[]
-        {
-            "NNE",
+		PrivateDependencyModuleNames.AddRange(new string[]
+		{
+			"NNE",
 			"NNEHlslShaders",
-            "RHI",
+			"RHI",
 			"Projects",
 			"TraceLog"
 		});
 
-        if (Target.Platform == UnrealTargetPlatform.Win64)
-		{
-			PrivateDependencyModuleNames.Add("D3D12RHI");
-			PrivateDependencyModuleNames.Add("DirectML");
-
-			AddEngineThirdPartyPrivateStaticDependencies(Target, "DX12");
-			AddEngineThirdPartyPrivateStaticDependencies(Target, "DirectML");
-
-			PublicDefinitions.Add("NNE_USE_DIRECTML");
-		}
-
 		if (Target.Platform == UnrealTargetPlatform.Mac)
 		{	
 			PrivateDependencyModuleNames.Add("MetalRHI");
diff --git a/Engine/Plugins/Experimental/NNERuntimeRDG/Source/NNERuntimeRDG/Private/Dml/NNEDmlCommon.h b/Engine/Plugins/Experimental/NNERuntimeRDG/Source/NNERuntimeRDG/Private/Dml/NNEDmlCommon.h
deleted file mode 100644
diff --git a/Engine/Plugins/Experimental/NNERuntimeRDG/Source/NNERuntimeRDG/Private/Dml/NNEDmlModel.cpp b/Engine/Plugins/Experimental/NNERuntimeRDG/Source/NNERuntimeRDG/Private/Dml/NNEDmlModel.cpp
deleted file mode 100644
diff --git a/Engine/Plugins/Experimental/NNERuntimeRDG/Source/NNERuntimeRDG/Private/Dml/NNEDmlModel.h b/Engine/Plugins/Experimental/NNERuntimeRDG/Source/NNERuntimeRDG/Private/Dml/NNEDmlModel.h
deleted file mode 100644
diff --git a/Engine/Plugins/Experimental/NNERuntimeRDG/Source/NNERuntimeRDG/Private/Dml/NNEDmlOperator.cpp b/Engine/Plugins/Experimental/NNERuntimeRDG/Source/NNERuntimeRDG/Private/Dml/NNEDmlOperator.cpp
deleted file mode 100644
diff --git a/Engine/Plugins/Experimental/NNERuntimeRDG/Source/NNERuntimeRDG/Private/Dml/NNEDmlOperator.h b/Engine/Plugins/Experimental/NNERuntimeRDG/Source/NNERuntimeRDG/Private/Dml/NNEDmlOperator.h
deleted file mode 100644
diff --git a/Engine/Plugins/Experimental/NNERuntimeRDG/Source/NNERuntimeRDG/Private/Dml/NNEDmlOperatorActivations.cpp b/Engine/Plugins/Experimental/NNERuntimeRDG/Source/NNERuntimeRDG/Private/Dml/NNEDmlOperatorActivations.cpp
deleted file mode 100644
diff --git a/Engine/Plugins/Experimental/NNERuntimeRDG/Source/NNERuntimeRDG/Private/Dml/NNEDmlOperatorBatchNormalization.cpp b/Engine/Plugins/Experimental/NNERuntimeRDG/Source/NNERuntimeRDG/Private/Dml/NNEDmlOperatorBatchNormalization.cpp
deleted file mode 100644
diff --git a/Engine/Plugins/Experimental/NNERuntimeRDG/Source/NNERuntimeRDG/Private/Dml/NNEDmlOperatorCast.cpp b/Engine/Plugins/Experimental/NNERuntimeRDG/Source/NNERuntimeRDG/Private/Dml/NNEDmlOperatorCast.cpp
deleted file mode 100644
diff --git a/Engine/Plugins/Experimental/NNERuntimeRDG/Source/NNERuntimeRDG/Private/Dml/NNEDmlOperatorConcat.cpp b/Engine/Plugins/Experimental/NNERuntimeRDG/Source/NNERuntimeRDG/Private/Dml/NNEDmlOperatorConcat.cpp
deleted file mode 100644
diff --git a/Engine/Plugins/Experimental/NNERuntimeRDG/Source/NNERuntimeRDG/Private/Dml/NNEDmlOperatorConvolution.cpp b/Engine/Plugins/Experimental/NNERuntimeRDG/Source/NNERuntimeRDG/Private/Dml/NNEDmlOperatorConvolution.cpp
deleted file mode 100644
diff --git a/Engine/Plugins/Experimental/NNERuntimeRDG/Source/NNERuntimeRDG/Private/Dml/NNEDmlOperatorCumSum.cpp b/Engine/Plugins/Experimental/NNERuntimeRDG/Source/NNERuntimeRDG/Private/Dml/NNEDmlOperatorCumSum.cpp
deleted file mode 100644
diff --git a/Engine/Plugins/Experimental/NNERuntimeRDG/Source/NNERuntimeRDG/Private/Dml/NNEDmlOperatorDepthToSpace.cpp b/Engine/Plugins/Experimental/NNERuntimeRDG/Source/NNERuntimeRDG/Private/Dml/NNEDmlOperatorDepthToSpace.cpp
deleted file mode 100644
diff --git a/Engine/Plugins/Experimental/NNERuntimeRDG/Source/NNERuntimeRDG/Private/Dml/NNEDmlOperatorElementwise.cpp b/Engine/Plugins/Experimental/NNERuntimeRDG/Source/NNERuntimeRDG/Private/Dml/NNEDmlOperatorElementwise.cpp
deleted file mode 100644
diff --git a/Engine/Plugins/Experimental/NNERuntimeRDG/Source/NNERuntimeRDG/Private/Dml/NNEDmlOperatorFlatten.cpp b/Engine/Plugins/Experimental/NNERuntimeRDG/Source/NNERuntimeRDG/Private/Dml/NNEDmlOperatorFlatten.cpp
deleted file mode 100644
diff --git a/Engine/Plugins/Experimental/NNERuntimeRDG/Source/NNERuntimeRDG/Private/Dml/NNEDmlOperatorGather.cpp b/Engine/Plugins/Experimental/NNERuntimeRDG/Source/NNERuntimeRDG/Private/Dml/NNEDmlOperatorGather.cpp
deleted file mode 100644
diff --git a/Engine/Plugins/Experimental/NNERuntimeRDG/Source/NNERuntimeRDG/Private/Dml/NNEDmlOperatorGemm.cpp b/Engine/Plugins/Experimental/NNERuntimeRDG/Source/NNERuntimeRDG/Private/Dml/NNEDmlOperatorGemm.cpp
deleted file mode 100644
diff --git a/Engine/Plugins/Experimental/NNERuntimeRDG/Source/NNERuntimeRDG/Private/Dml/NNEDmlOperatorIdentity.cpp b/Engine/Plugins/Experimental/NNERuntimeRDG/Source/NNERuntimeRDG/Private/Dml/NNEDmlOperatorIdentity.cpp
deleted file mode 100644
diff --git a/Engine/Plugins/Experimental/NNERuntimeRDG/Source/NNERuntimeRDG/Private/Dml/NNEDmlOperatorInstanceNormalization.cpp b/Engine/Plugins/Experimental/NNERuntimeRDG/Source/NNERuntimeRDG/Private/Dml/NNEDmlOperatorInstanceNormalization.cpp
deleted file mode 100644
diff --git a/Engine/Plugins/Experimental/NNERuntimeRDG/Source/NNERuntimeRDG/Private/Dml/NNEDmlOperatorLRN.cpp b/Engine/Plugins/Experimental/NNERuntimeRDG/Source/NNERuntimeRDG/Private/Dml/NNEDmlOperatorLRN.cpp
deleted file mode 100644
diff --git a/Engine/Plugins/Experimental/NNERuntimeRDG/Source/NNERuntimeRDG/Private/Dml/NNEDmlOperatorLpNormalization.cpp b/Engine/Plugins/Experimental/NNERuntimeRDG/Source/NNERuntimeRDG/Private/Dml/NNEDmlOperatorLpNormalization.cpp
deleted file mode 100644
diff --git a/Engine/Plugins/Experimental/NNERuntimeRDG/Source/NNERuntimeRDG/Private/Dml/NNEDmlOperatorMatMul.cpp b/Engine/Plugins/Experimental/NNERuntimeRDG/Source/NNERuntimeRDG/Private/Dml/NNEDmlOperatorMatMul.cpp
deleted file mode 100644
diff --git a/Engine/Plugins/Experimental/NNERuntimeRDG/Source/NNERuntimeRDG/Private/Dml/NNEDmlOperatorMaxUnpool.cpp b/Engine/Plugins/Experimental/NNERuntimeRDG/Source/NNERuntimeRDG/Private/Dml/NNEDmlOperatorMaxUnpool.cpp
deleted file mode 100644
diff --git a/Engine/Plugins/Experimental/NNERuntimeRDG/Source/NNERuntimeRDG/Private/Dml/NNEDmlOperatorMeanVarianceNormalization.cpp b/Engine/Plugins/Experimental/NNERuntimeRDG/Source/NNERuntimeRDG/Private/Dml/NNEDmlOperatorMeanVarianceNormalization.cpp
deleted file mode 100644
diff --git a/Engine/Plugins/Experimental/NNERuntimeRDG/Source/NNERuntimeRDG/Private/Dml/NNEDmlOperatorPad.cpp b/Engine/Plugins/Experimental/NNERuntimeRDG/Source/NNERuntimeRDG/Private/Dml/NNEDmlOperatorPad.cpp
deleted file mode 100644
diff --git a/Engine/Plugins/Experimental/NNERuntimeRDG/Source/NNERuntimeRDG/Private/Dml/NNEDmlOperatorPool.cpp b/Engine/Plugins/Experimental/NNERuntimeRDG/Source/NNERuntimeRDG/Private/Dml/NNEDmlOperatorPool.cpp
deleted file mode 100644
diff --git a/Engine/Plugins/Experimental/NNERuntimeRDG/Source/NNERuntimeRDG/Private/Dml/NNEDmlOperatorReduce.cpp b/Engine/Plugins/Experimental/NNERuntimeRDG/Source/NNERuntimeRDG/Private/Dml/NNEDmlOperatorReduce.cpp
deleted file mode 100644
diff --git a/Engine/Plugins/Experimental/NNERuntimeRDG/Source/NNERuntimeRDG/Private/Dml/NNEDmlOperatorResample.cpp b/Engine/Plugins/Experimental/NNERuntimeRDG/Source/NNERuntimeRDG/Private/Dml/NNEDmlOperatorResample.cpp
deleted file mode 100644
diff --git a/Engine/Plugins/Experimental/NNERuntimeRDG/Source/NNERuntimeRDG/Private/Dml/NNEDmlOperatorReshape.cpp b/Engine/Plugins/Experimental/NNERuntimeRDG/Source/NNERuntimeRDG/Private/Dml/NNEDmlOperatorReshape.cpp
deleted file mode 100644
diff --git a/Engine/Plugins/Experimental/NNERuntimeRDG/Source/NNERuntimeRDG/Private/Dml/NNEDmlOperatorSlice.cpp b/Engine/Plugins/Experimental/NNERuntimeRDG/Source/NNERuntimeRDG/Private/Dml/NNEDmlOperatorSlice.cpp
deleted file mode 100644
diff --git a/Engine/Plugins/Experimental/NNERuntimeRDG/Source/NNERuntimeRDG/Private/Dml/NNEDmlOperatorSpaceToDepth.cpp b/Engine/Plugins/Experimental/NNERuntimeRDG/Source/NNERuntimeRDG/Private/Dml/NNEDmlOperatorSpaceToDepth.cpp
deleted file mode 100644
diff --git a/Engine/Plugins/Experimental/NNERuntimeRDG/Source/NNERuntimeRDG/Private/Dml/NNEDmlOperatorSplit.cpp b/Engine/Plugins/Experimental/NNERuntimeRDG/Source/NNERuntimeRDG/Private/Dml/NNEDmlOperatorSplit.cpp
deleted file mode 100644
diff --git a/Engine/Plugins/Experimental/NNERuntimeRDG/Source/NNERuntimeRDG/Private/Dml/NNEDmlOperatorSqueeze.cpp b/Engine/Plugins/Experimental/NNERuntimeRDG/Source/NNERuntimeRDG/Private/Dml/NNEDmlOperatorSqueeze.cpp
deleted file mode 100644
diff --git a/Engine/Plugins/Experimental/NNERuntimeRDG/Source/NNERuntimeRDG/Private/Dml/NNEDmlOperatorTranspose.cpp b/Engine/Plugins/Experimental/NNERuntimeRDG/Source/NNERuntimeRDG/Private/Dml/NNEDmlOperatorTranspose.cpp
deleted file mode 100644
diff --git a/Engine/Plugins/Experimental/NNERuntimeRDG/Source/NNERuntimeRDG/Private/Dml/NNEDmlOperatorUnsqueeze.cpp b/Engine/Plugins/Experimental/NNERuntimeRDG/Source/NNERuntimeRDG/Private/Dml/NNEDmlOperatorUnsqueeze.cpp
deleted file mode 100644
diff --git a/Engine/Plugins/Experimental/NNERuntimeRDG/Source/NNERuntimeRDG/Private/Dml/NNEDmlOperatorUtils.h b/Engine/Plugins/Experimental/NNERuntimeRDG/Source/NNERuntimeRDG/Private/Dml/NNEDmlOperatorUtils.h
deleted file mode 100644
diff --git a/Engine/Plugins/Experimental/NNERuntimeRDG/Source/NNERuntimeRDG/Private/NNERuntimeRDGDml.cpp b/Engine/Plugins/Experimental/NNERuntimeRDG/Source/NNERuntimeRDG/Private/NNERuntimeRDGDml.cpp
deleted file mode 100644
diff --git a/Engine/Plugins/Experimental/NNERuntimeRDG/Source/NNERuntimeRDG/Private/NNERuntimeRDGDml.h b/Engine/Plugins/Experimental/NNERuntimeRDG/Source/NNERuntimeRDG/Private/NNERuntimeRDGDml.h
deleted file mode 100644
diff --git a/Engine/Plugins/Experimental/NNERuntimeRDG/Source/NNERuntimeRDG/Private/NNERuntimeRDGModule.cpp b/Engine/Plugins/Experimental/NNERuntimeRDG/Source/NNERuntimeRDG/Private/NNERuntimeRDGModule.cpp
--- a/Engine/Plugins/Experimental/NNERuntimeRDG/Source/NNERuntimeRDG/Private/NNERuntimeRDGModule.cpp
+++ b/Engine/Plugins/Experimental/NNERuntimeRDG/Source/NNERuntimeRDG/Private/NNERuntimeRDGModule.cpp
@@ -2,7 +2,6 @@
 
 #include "NNERuntimeRDGModule.h"
 #include "NNE.h"
-#include "NNERuntimeRDGDml.h"
 #include "NNERuntimeRDGHlsl.h"
 #include "UObject/WeakInterfacePtr.h"
 
@@ -19,20 +18,6 @@ void FNNERuntimeRDGModule::StartupModule()
 		NNERuntimeRDGHlsl->AddToRoot();
 		UE::NNE::RegisterRuntime(RuntimeCPUInterface);
 	}
-
-#ifdef NNE_USE_DIRECTML
-	// NNE runtime ORT Dml startup
-	NNERuntimeRDGDml = NewObject<UNNERuntimeRDGDmlImpl>();
-	if (NNERuntimeRDGDml.IsValid())
-	{
-		TWeakInterfacePtr<INNERuntime> RuntimeDmlInterface(NNERuntimeRDGDml.Get());
-		
-		bool bRegisterOnlyOperators = UE::NNERuntimeRDG::Private::Dml::FRuntimeDmlStartup();
-		NNERuntimeRDGDml->Init(bRegisterOnlyOperators);
-		NNERuntimeRDGDml->AddToRoot();
-		UE::NNE::RegisterRuntime(RuntimeDmlInterface);
-	}
-#endif
 }
 
 void FNNERuntimeRDGModule::ShutdownModule()
@@ -46,19 +31,6 @@ void FNNERuntimeRDGModule::ShutdownModule()
 		NNERuntimeRDGHlsl->RemoveFromRoot();
 		NNERuntimeRDGHlsl.Reset();
 	}
-
-#ifdef NNE_USE_DIRECTML
-	// NNE runtime ORT Dml shutdown
-	if (NNERuntimeRDGDml.IsValid())
-	{
-		TWeakInterfacePtr<INNERuntime> RuntimeDmlInterface(NNERuntimeRDGDml.Get());
-
-		UE::NNE::UnregisterRuntime(RuntimeDmlInterface);
-		NNERuntimeRDGDml->RemoveFromRoot();
-		NNERuntimeRDGDml.Reset();
-	}
-#endif
-
 }
 
 IMPLEMENT_MODULE(FNNERuntimeRDGModule, NNERuntimeRDG);
diff --git a/Engine/Plugins/Experimental/NNERuntimeRDG/Source/NNERuntimeRDG/Private/NNERuntimeRDGModule.h b/Engine/Plugins/Experimental/NNERuntimeRDG/Source/NNERuntimeRDG/Private/NNERuntimeRDGModule.h
--- a/Engine/Plugins/Experimental/NNERuntimeRDG/Source/NNERuntimeRDG/Private/NNERuntimeRDGModule.h
+++ b/Engine/Plugins/Experimental/NNERuntimeRDG/Source/NNERuntimeRDG/Private/NNERuntimeRDGModule.h
@@ -8,7 +8,6 @@
 #include "UObject/WeakObjectPtr.h"
 #include "UObject/WeakObjectPtrTemplates.h"
 
-class UNNERuntimeRDGDmlImpl;
 class UNNERuntimeRDGHlslImpl;
 
 class FNNERuntimeRDGModule : public IModuleInterface
@@ -16,7 +15,6 @@ class FNNERuntimeRDGModule : public IModuleInterface
 
 public:
 	TWeakObjectPtr<UNNERuntimeRDGHlslImpl> NNERuntimeRDGHlsl{ nullptr };
-	TWeakObjectPtr<UNNERuntimeRDGDmlImpl> NNERuntimeRDGDml{ nullptr };
 
 	// Begin IModuleInterface
 	virtual void StartupModule() override;
