[NNE] Extend Hlsl Reduce support to very large tensors

Problems:
* Hlsl Reduce implementation runs into a d3dx12 limitation for a number of elements > 65535 before/after the axis.

Solutions:
* Modify Hlsl Reduce implementation to issue multiple dispatches and overcome such limitation

Tests:
* Hlsl Reduce all tests
* Win64 Standalone

#rnx

#virtualized

diff --git a/Engine/Plugins/Experimental/NNERuntimeRDG/Source/NNEHlslShaders/Internal/NNEHlslShadersReduceCS.h b/Engine/Plugins/Experimental/NNERuntimeRDG/Source/NNEHlslShaders/Internal/NNEHlslShadersReduceCS.h
--- a/Engine/Plugins/Experimental/NNERuntimeRDG/Source/NNEHlslShaders/Internal/NNEHlslShadersReduceCS.h
+++ b/Engine/Plugins/Experimental/NNERuntimeRDG/Source/NNEHlslShaders/Internal/NNEHlslShadersReduceCS.h
@@ -7,6 +7,7 @@
 #include "RenderGraphFwd.h"
 #include "RenderGraphUtils.h"
 #include "ShaderParameterUtils.h"
+#include "RHIGlobals.h"
 
 namespace UE::NNE::Internal { class FTensor; }
 
@@ -30,7 +31,7 @@ namespace UE::NNEHlslShaders::Internal
 	class FReduceConstants
 	{
 	public:
-		static const int32 NUM_GROUP_THREADS{ 768 };
+		static const int32 THREADGROUP_SIZE{ 768 };
 	};
 
 	class NNEHLSLSHADERS_API TReduceCS : public FHlslShaderBase
@@ -47,6 +48,7 @@ namespace UE::NNEHlslShaders::Internal
 			SHADER_PARAMETER(int32, NumElemBeforeAxis)
 			SHADER_PARAMETER(int32, AxisSize)
 			SHADER_PARAMETER(int32, NumElemAfterAxis)
+			SHADER_PARAMETER(FUintVector4, DispatchIdxAndStride)
 			SHADER_PARAMETER(float, Epsilon)
 			SHADER_PARAMETER_RDG_BUFFER_SRV(Buffer<float>, Input)
 			SHADER_PARAMETER_RDG_BUFFER_UAV(RWBuffer<float>, Output)
diff --git a/Engine/Plugins/Experimental/NNERuntimeRDG/Source/NNEHlslShaders/Private/NNEHlslShadersReduceCS.cpp b/Engine/Plugins/Experimental/NNERuntimeRDG/Source/NNEHlslShaders/Private/NNEHlslShadersReduceCS.cpp
--- a/Engine/Plugins/Experimental/NNERuntimeRDG/Source/NNEHlslShaders/Private/NNEHlslShadersReduceCS.cpp
+++ b/Engine/Plugins/Experimental/NNERuntimeRDG/Source/NNEHlslShaders/Private/NNEHlslShadersReduceCS.cpp
@@ -10,7 +10,7 @@ namespace UE::NNEHlslShaders::Internal
 	{
 		FGlobalShader::ModifyCompilationEnvironment(InParameters, OutEnvironment);
 
-		OutEnvironment.SetDefine(TEXT("THREADGROUP_SIZE"), FReduceConstants::NUM_GROUP_THREADS);
+		OutEnvironment.SetDefine(TEXT("THREADGROUP_SIZE"), FReduceConstants::THREADGROUP_SIZE);
 	}
 
 	void TReduceCS::FillInParameters(TConstArrayView<uint32> Shape, int32 Axis, TReduceCS::FParameters* Parameters)
@@ -35,29 +35,54 @@ namespace UE::NNEHlslShaders::Internal
 		Parameters->Epsilon = 0;
 	}
 
-	void TReduceCS::EnqueueRDG(FRDGBuilder& GraphBuilder, TReduceCS::FParameters* Parameters, FRDGBufferRef Input, FRDGBufferRef Output, EReduceOperatorType OperatorType, FRDGBufferRef Output2)
+	void TReduceCS::EnqueueRDG(FRDGBuilder& GraphBuilder, TReduceCS::FParameters* InParameters, FRDGBufferRef Input, FRDGBufferRef Output, EReduceOperatorType OperatorType, FRDGBufferRef Output2)
 	{
-		check(Parameters);
+		check(InParameters);
 
-		Parameters->Input = GraphBuilder.CreateSRV(FRDGBufferSRVDesc(Input, PF_R32_FLOAT));
-		Parameters->Output = GraphBuilder.CreateUAV(FRDGBufferUAVDesc(Output, PF_R32_FLOAT));
+		InParameters->Input = GraphBuilder.CreateSRV(FRDGBufferSRVDesc(Input, PF_R32_FLOAT));
+		InParameters->Output = GraphBuilder.CreateUAV(FRDGBufferUAVDesc(Output, PF_R32_FLOAT));
 		if (Output2 != nullptr) {
-			Parameters->Output2 = GraphBuilder.CreateUAV(FRDGBufferUAVDesc(Output2, PF_R32_FLOAT));
+			InParameters->Output2 = GraphBuilder.CreateUAV(FRDGBufferUAVDesc(Output2, PF_R32_FLOAT));
 		}
 
 		TReduceCS::FPermutationDomain PermutationVector;
 		PermutationVector.Set<TReduceCS::FReduceType>(OperatorType);
 
 		TShaderMapRef<TReduceCS> ComputeShader(GetGlobalShaderMap(GMaxRHIFeatureLevel), PermutationVector);
-		FIntVector ThreadGroupCount{ 1, (int32)Parameters->NumElemAfterAxis, (int32)Parameters->NumElemBeforeAxis };
 
-		FComputeShaderUtils::AddPass(
-			GraphBuilder,
-			RDG_EVENT_NAME("NNE.Operator.Hlsl.Reduce.OneAxis.Dispatch"),
-			ERDGPassFlags::Compute | ERDGPassFlags::NeverCull,
-			ComputeShader,
-			Parameters,
-			ThreadGroupCount);
+		FIntVector3 NumDispatches (
+			1,
+			(int32)InParameters->NumElemAfterAxis / GRHIMaxDispatchThreadGroupsPerDimension.Y + 1,
+			(int32)InParameters->NumElemBeforeAxis / GRHIMaxDispatchThreadGroupsPerDimension.Z + 1
+			);
+		
+		for(int32 DispatchIdxY = 0; DispatchIdxY < NumDispatches.Y; ++DispatchIdxY)
+		{
+			for(int32 DispatchIdxZ = 0; DispatchIdxZ < NumDispatches.Z; ++DispatchIdxZ)
+			{
+				FIntVector ThreadGroupCount {
+					1,
+					DispatchIdxY == NumDispatches.Y - 1 ? (int32)InParameters->NumElemAfterAxis % GRHIMaxDispatchThreadGroupsPerDimension.Y : GRHIMaxDispatchThreadGroupsPerDimension.Y,
+					DispatchIdxZ == NumDispatches.Z - 1 ? (int32)InParameters->NumElemBeforeAxis % GRHIMaxDispatchThreadGroupsPerDimension.Z : GRHIMaxDispatchThreadGroupsPerDimension.Z
+				};
+
+				// Need to use a copy of InParameters to set DispatchIdxAndStride for a single pass
+				TReduceCS::FParameters* CurParameters = GraphBuilder.AllocParameters(InParameters);
+				CurParameters->DispatchIdxAndStride.X = (uint32) DispatchIdxY;
+				CurParameters->DispatchIdxAndStride.Y = (uint32) DispatchIdxZ;
+				CurParameters->DispatchIdxAndStride.Z = (uint32) GRHIMaxDispatchThreadGroupsPerDimension.Y;
+				CurParameters->DispatchIdxAndStride.W = (uint32) GRHIMaxDispatchThreadGroupsPerDimension.Z;
+				
+				FComputeShaderUtils::AddPass(
+					GraphBuilder,
+					RDG_EVENT_NAME("NNE.Operator.Hlsl.Reduce.OneAxis.Dispatch.%d.%d", DispatchIdxY, DispatchIdxZ),
+					ERDGPassFlags::Compute | ERDGPassFlags::NeverCull,
+					ComputeShader,
+					CurParameters,
+					ThreadGroupCount);
+			}
+		}
+		
 	}
 
 	IMPLEMENT_GLOBAL_SHADER(TReduceCS, "/NNE/NNEHlslShadersReduce.usf", "Reduce", SF_Compute);
diff --git a/Engine/Plugins/Experimental/NNERuntimeRDG/Source/NNEHlslShaders/Shaders/NNEHlslShadersReduce.usf b/Engine/Plugins/Experimental/NNERuntimeRDG/Source/NNEHlslShaders/Shaders/NNEHlslShadersReduce.usf
--- a/Engine/Plugins/Experimental/NNERuntimeRDG/Source/NNEHlslShaders/Shaders/NNEHlslShadersReduce.usf
+++ b/Engine/Plugins/Experimental/NNERuntimeRDG/Source/NNEHlslShaders/Shaders/NNEHlslShadersReduce.usf
@@ -13,6 +13,7 @@ RWBuffer<BUFFER_TYPE> Output2;
 int NumElemBeforeAxis;
 int AxisSize;
 int NumElemAfterAxis;
+uint4 DispatchIdxAndStride;
 float Epsilon;
 
 // Must correspond to EReduceOperatorType defined in NNEHlslShadersReduceCS.h
@@ -130,8 +131,8 @@ float Epsilon;
 [numthreads(THREADGROUP_SIZE, 1, 1)]
 void Reduce(in const uint3 DispatchThreadID : SV_DispatchThreadID, in const uint3 GroupThreadID : SV_GroupThreadID)
 {
-	const uint IdxAfterAxis = DispatchThreadID.y;
-	const uint IdxBeforeAxis = DispatchThreadID.z;
+	const uint IdxAfterAxis = DispatchIdxAndStride.x * DispatchIdxAndStride.z +  DispatchThreadID.y;
+	const uint IdxBeforeAxis =  DispatchIdxAndStride.y * DispatchIdxAndStride.w +  DispatchThreadID.z;
 	const int InputStartIdx = (IdxBeforeAxis * AxisSize * NumElemAfterAxis) + IdxAfterAxis;
 	const int OutputIdx     = (IdxBeforeAxis * 1        * NumElemAfterAxis) + IdxAfterAxis;
 
