NNE - Add Hlsl LayerNormalization

Problems:
* Hlsl LayerNormalization implementation is missing

Solutions:
* Implement Hlsl LayerNormalization

Tests:
* LayerNormalization tests
* Smoke tests

#rnx

#virtualized

diff --git a/Engine/Plugins/Experimental/NNERuntimeRDG/Source/NNEHlslShaders/Internal/NNEHlslShadersLayerNormalizationCS.h b/Engine/Plugins/Experimental/NNERuntimeRDG/Source/NNEHlslShaders/Internal/NNEHlslShadersLayerNormalizationCS.h
new file mode 100644
--- /dev/null
+++ b/Engine/Plugins/Experimental/NNERuntimeRDG/Source/NNEHlslShaders/Internal/NNEHlslShadersLayerNormalizationCS.h
@@ -0,0 +1,51 @@
+// Copyright Epic Games, Inc. All Rights Reserved.
+
+#pragma once
+
+#include "NNEHlslShaderBase.h"
+#include "RenderGraphUtils.h"
+#include "ShaderParameterUtils.h"
+
+namespace UE::NNE::Internal { class FTensor; }
+
+namespace UE::NNEHlslShaders::Internal
+{
+	class FLayerNormalizationConstants
+	{
+	public:
+		static const int32 MAX_NUM_DIMENSIONS{ 8 };
+		static const int32 NUM_GROUP_THREADS{ 768 };
+	};
+
+	class NNEHLSLSHADERS_API TLayerNormalizationCS : public FHlslShaderBase
+	{
+		DECLARE_GLOBAL_SHADER(TLayerNormalizationCS);
+		SHADER_USE_PARAMETER_STRUCT(TLayerNormalizationCS, FHlslShaderBase)
+
+		class FLayerNormalizationNumDimensions : SHADER_PERMUTATION_RANGE_INT("NUM_DIMENSIONS", 1, FLayerNormalizationConstants::MAX_NUM_DIMENSIONS);
+		class FLayerNormalizationHasB : SHADER_PERMUTATION_BOOL("HAS_B");
+		class FLayerNormalizationWriteMean : SHADER_PERMUTATION_BOOL("WRITE_MEAN");
+		class FLayerNormalizationWriteInvStdDev : SHADER_PERMUTATION_BOOL("WRITE_INVSTDDEV");
+		using FPermutationDomain = TShaderPermutationDomain<FLayerNormalizationNumDimensions, FLayerNormalizationHasB, FLayerNormalizationWriteMean, FLayerNormalizationWriteInvStdDev>;
+
+	public:
+
+		BEGIN_SHADER_PARAMETER_STRUCT(FParameters, )
+			SHADER_PARAMETER(uint32, Axis)
+			SHADER_PARAMETER(float, Epsilon)
+			SHADER_PARAMETER(uint32, LayerSize)
+			SHADER_PARAMETER_RDG_BUFFER_SRV(Buffer<float>, Input)
+			SHADER_PARAMETER_ARRAY(FUintVector4, InputTensorInfo, [FLayerNormalizationConstants::MAX_NUM_DIMENSIONS])
+			SHADER_PARAMETER_RDG_BUFFER_SRV(Buffer<float>, InputScale)
+			SHADER_PARAMETER_ARRAY(FUintVector4, ScaleTensorInfo, [FLayerNormalizationConstants::MAX_NUM_DIMENSIONS])
+			SHADER_PARAMETER_RDG_BUFFER_SRV(Buffer<float>, InputBias)
+			SHADER_PARAMETER_ARRAY(FUintVector4, BiasTensorInfo, [FLayerNormalizationConstants::MAX_NUM_DIMENSIONS])
+			SHADER_PARAMETER_RDG_BUFFER_UAV(RWBuffer<float>, Output)
+			SHADER_PARAMETER_RDG_BUFFER_UAV(RWBuffer<float>, OutputMean)
+			SHADER_PARAMETER_RDG_BUFFER_UAV(RWBuffer<float>, OutputInvStdDev)
+			END_SHADER_PARAMETER_STRUCT()
+
+			static void FillInParameters(TConstArrayView<uint32> Shape, int32 Axis, float Epsilon, FParameters* Parameters);
+			static void ModifyCompilationEnvironment(const FGlobalShaderPermutationParameters& InParameters, FShaderCompilerEnvironment& OutEnvironment);
+	};
+} // UE::NNEHlslShaders::Internal
\ No newline at end of file
diff --git a/Engine/Plugins/Experimental/NNERuntimeRDG/Source/NNEHlslShaders/Private/NNEHlslShadersLayerNormalizationCS.cpp b/Engine/Plugins/Experimental/NNERuntimeRDG/Source/NNEHlslShaders/Private/NNEHlslShadersLayerNormalizationCS.cpp
new file mode 100644
--- /dev/null
+++ b/Engine/Plugins/Experimental/NNERuntimeRDG/Source/NNEHlslShaders/Private/NNEHlslShadersLayerNormalizationCS.cpp
@@ -0,0 +1,34 @@
+// Copyright Epic Games, Inc. All Rights Reserved.
+
+#include "NNEHlslShadersLayerNormalizationCS.h"
+#include "NNE.h"
+#include "NNETensor.h"
+
+namespace UE::NNEHlslShaders::Internal
+{
+	void TLayerNormalizationCS::ModifyCompilationEnvironment(const FGlobalShaderPermutationParameters& InParameters, FShaderCompilerEnvironment& OutEnvironment)
+	{
+		FGlobalShader::ModifyCompilationEnvironment(InParameters, OutEnvironment);
+		OutEnvironment.SetDefine(TEXT("THREADGROUP_SIZE"), FLayerNormalizationConstants::NUM_GROUP_THREADS);
+
+		FPermutationDomain PermutationVector(InParameters.PermutationId);
+	}
+
+	void TLayerNormalizationCS::FillInParameters(TConstArrayView<uint32> Shape, int32 Axis, float Epsilon, TLayerNormalizationCS::FParameters* Parameters)
+	{
+		check(Axis >= 0 && Axis < Shape.Num());
+		check(Parameters);
+
+		uint32 LayerSize = 1;
+		for(int DimIdx = Axis; DimIdx < Shape.Num(); ++DimIdx)
+		{
+			LayerSize *= Shape[DimIdx];
+		}
+
+		Parameters->LayerSize = LayerSize;
+		Parameters->Axis = Axis;
+		Parameters->Epsilon = Epsilon;
+	}
+
+	IMPLEMENT_GLOBAL_SHADER(TLayerNormalizationCS, "/NNE/NNEHlslShadersLayerNormalization.usf", "LayerNormalization", SF_Compute);
+} // UE::NNEHlslShaders::Internal
\ No newline at end of file
diff --git a/Engine/Plugins/Experimental/NNERuntimeRDG/Source/NNEHlslShaders/Shaders/NNEHlslShadersLayerNormalization.usf b/Engine/Plugins/Experimental/NNERuntimeRDG/Source/NNEHlslShaders/Shaders/NNEHlslShadersLayerNormalization.usf
new file mode 100644
--- /dev/null
+++ b/Engine/Plugins/Experimental/NNERuntimeRDG/Source/NNEHlslShaders/Shaders/NNEHlslShadersLayerNormalization.usf
@@ -0,0 +1,144 @@
+// Copyright Epic Games, Inc. All Rights Reserved.
+ 
+#include "/Engine/Public/Platform.ush"
+#include "NNEHlslShadersBroadcastHelper.ush"
+
+//NOTE: once have support for more datatypes, make BUFFER_TYPE into a variant and WORK_TYPE depend on stash_type
+#define WORK_TYPE float
+#define BUFFER_TYPE float
+#define READ(x) x
+#define WRITE(x) x
+
+Buffer<BUFFER_TYPE> Input;
+uint4 InputTensorInfo[NUM_DIMENSIONS];
+Buffer<BUFFER_TYPE> InputScale;
+uint4 ScaleTensorInfo[NUM_DIMENSIONS];
+Buffer<BUFFER_TYPE> InputBias;
+uint4 BiasTensorInfo[NUM_DIMENSIONS];
+RWBuffer<BUFFER_TYPE> Output;
+RWBuffer<BUFFER_TYPE> OutputMean;
+RWBuffer<BUFFER_TYPE> OutputInvStdDev;
+uint Axis;
+float Epsilon;
+uint LayerSize;
+
+#define STRIDE_IDX 0
+
+#define STATIC_LOOP(Var, From, To) \
+[unroll] \
+for(uint Var = From; Var < To; ++Var)
+
+void GetDimIterator(const uint GlobalIdx, const uint4 TensorInfo[NUM_DIMENSIONS], out uint DimIterator[NUM_DIMENSIONS])
+{
+	uint Offset = GlobalIdx;
+	STATIC_LOOP(DimIdx, 0, NUM_DIMENSIONS)
+	{
+		uint Remainder;
+		DivMod(Offset, TensorInfo[DimIdx][STRIDE_IDX], DimIterator[DimIdx], Remainder);
+		Offset = Remainder;
+	}
+}
+
+uint GetInnerIdx(const uint DimIterator[NUM_DIMENSIONS], const uint4 TensorInfo[NUM_DIMENSIONS], const uint Axis)
+{
+	uint Offset = 0;
+	for(uint DimIdx = 0; DimIdx < NUM_DIMENSIONS - Axis; ++DimIdx)
+	{
+		Offset += DimIterator[DimIdx + Axis] * TensorInfo[DimIdx][STRIDE_IDX];
+	}
+	return Offset;
+}
+
+groupshared WORK_TYPE SharedMemoryMean[THREADGROUP_SIZE];
+groupshared WORK_TYPE SharedMemoryMean2[THREADGROUP_SIZE];
+groupshared WORK_TYPE SharedMemoryCount[THREADGROUP_SIZE];
+
+[numthreads(THREADGROUP_SIZE, 1, 1)]
+void LayerNormalization(in const uint3 GroupID : SV_GroupID, in const uint3 GroupThreadID : SV_GroupThreadID, in const uint3 DispatchThreadID : SV_DispatchThreadID)
+{
+	const int SliceIdx = GroupID.x; // each thread group processes a slice
+	const int InputStartIdx = SliceIdx * LayerSize;
+	
+
+	// Online Welford using sharedmem and THREADGROUP_SIZE threads along the dimensions after `Axis` (included)
+	{
+		WORK_TYPE Mean = 0;
+		WORK_TYPE Mean2 = 0;
+		int Count = 0;
+
+		for (int InElemIdx = GroupThreadID.x; InElemIdx < LayerSize; InElemIdx += THREADGROUP_SIZE)
+		{
+			WORK_TYPE Value = (WORK_TYPE) READ(Input[InElemIdx + InputStartIdx]);
+		
+			Count += 1;
+			WORK_TYPE Delta = Value - Mean;
+			Mean += Delta / (WORK_TYPE)Count;
+			WORK_TYPE Delta2 = Value - Mean;
+			Mean2 += Delta * Delta2;
+		}
+
+		SharedMemoryMean[GroupThreadID.x] = Mean;
+		SharedMemoryMean2[GroupThreadID.x] = Mean2;
+		SharedMemoryCount[GroupThreadID.x] = Count;
+	}
+
+	GroupMemoryBarrierWithGroupSync();
+
+	if (GroupThreadID.x == 0)
+	{
+		WORK_TYPE Mean = 0;
+		WORK_TYPE Mean2 = 0;
+		int Count = 0;
+
+		UNROLL
+		for (int Idx = 0; Idx < THREADGROUP_SIZE; Idx++)
+		{
+			WORK_TYPE CurrMean = SharedMemoryMean[Idx];
+			WORK_TYPE CurrMean2 = SharedMemoryMean2[Idx];
+			int CurrCount = SharedMemoryCount[Idx];
+			int NewCount = Count + CurrCount;
+			WORK_TYPE Delta = CurrMean - Mean;
+			Mean = (CurrCount * CurrMean + Count * Mean) / NewCount;
+			Mean2 += CurrMean2 + Delta * Delta * (WORK_TYPE)Count * (WORK_TYPE)CurrCount / (WORK_TYPE)NewCount;
+			Count = NewCount;
+		}
+
+		SharedMemoryMean[0] = Mean;
+		SharedMemoryMean2[0] = Mean2;
+		SharedMemoryCount[0] = Count;
+		
+	}
+
+	GroupMemoryBarrierWithGroupSync();
+
+	WORK_TYPE Mean = SharedMemoryMean[0];
+	WORK_TYPE Mean2 = SharedMemoryMean2[0];
+	int Count = SharedMemoryCount[0];
+
+	WORK_TYPE Variance = Mean2 / (WORK_TYPE)Count;
+	WORK_TYPE InvStdDev = (WORK_TYPE)1 / sqrt(Variance + Epsilon);
+
+	for (int InElemIdx = GroupThreadID.x; InElemIdx < LayerSize; InElemIdx += THREADGROUP_SIZE)
+	{
+		const int ElementIdx = InElemIdx + InputStartIdx;
+
+		BUFFER_TYPE Normalized = (BUFFER_TYPE) ( (READ(Input[ElementIdx]) - Mean) * InvStdDev );
+
+		uint InputDimIterator[NUM_DIMENSIONS];
+		GetDimIterator(ElementIdx, InputTensorInfo, InputDimIterator);
+
+		BUFFER_TYPE NormalizedScaled = Normalized * READ( InputScale[GetInnerIdx(InputDimIterator, ScaleTensorInfo, Axis)]);
+#if HAS_B == 0
+		Output[ElementIdx] = WRITE( NormalizedScaled );
+#else
+		Output[ElementIdx] = WRITE( NormalizedScaled + READ( InputBias[GetInnerIdx(InputDimIterator, BiasTensorInfo, Axis)] ) );
+#endif
+
+	}
+#if WRITE_MEAN != 0
+	OutputMean[SliceIdx] = WRITE( (BUFFER_TYPE) Mean );
+#endif
+#if WRITE_INVSTDDEV != 0
+	OutputInvStdDev[SliceIdx] = WRITE( (BUFFER_TYPE) InvStdDev );
+#endif
+}
diff --git a/Engine/Plugins/Experimental/NNERuntimeRDG/Source/NNERuntimeRDG/Private/Hlsl/NNERuntimeRDGLayerNormalization.cpp b/Engine/Plugins/Experimental/NNERuntimeRDG/Source/NNERuntimeRDG/Private/Hlsl/NNERuntimeRDGLayerNormalization.cpp
new file mode 100644
--- /dev/null
+++ b/Engine/Plugins/Experimental/NNERuntimeRDG/Source/NNERuntimeRDG/Private/Hlsl/NNERuntimeRDGLayerNormalization.cpp
@@ -0,0 +1,243 @@
+// Copyright Epic Games, Inc. All Rights Reserved.
+
+#include "NNERuntimeRDGLayerNormalization.h"
+#include "NNEHlslShadersLayerNormalizationCS.h"
+#include "NNERuntimeRDGHlslHelper.h"
+#include "NNETensor.h"
+#include "NNETypes.h"
+
+namespace UE::NNERuntimeRDG::Private::Hlsl
+{
+	DECLARE_GPU_STAT_NAMED(FNNEOperatorLayerNormalization, TEXT("NNE.Operator.Hlsl.LayerNormalization"));
+
+	/**
+	 * LayerNormalization operator implementation
+	 */
+	class TLayerNormalization : public FOperatorHlsl
+	{
+	public:
+
+		TLayerNormalization() {}
+		virtual ~TLayerNormalization() = default;
+
+	private:
+
+		float Epsilon = 1e-5f;
+		int32 Axis = -1;
+		int32 StashType = 1;
+
+	public:
+
+		virtual int PrepareOutputs(TConstArrayView<NNE::Internal::FTensorRef> InputTensors, TArrayView<NNE::Internal::FTensorRef> OutputTensors) override
+		{
+			check(InputTensors.Num() >= 2 && InputTensors.Num() <= 3);
+			check(OutputTensors.Num() >= 1 && OutputTensors.Num() <= 3);
+
+			auto ValidateScaleOrBias = [Axis = Axis] (const UE::NNE::FTensorShape& Input, const UE::NNE::FTensorShape& ScaleOrBias) -> int
+			{
+				if (ScaleOrBias.Rank() > (Input.Rank() - Axis))
+				{
+					UE_LOG(LogNNE, Warning, TEXT("LayerNormalization: Scale/bias tensor rank is invalid: %d"), ScaleOrBias.Rank());
+					return -1;
+				}
+				for(int DimIdx = 0; DimIdx < ScaleOrBias.Rank(); ++DimIdx)
+				{
+					if(Input.GetData()[DimIdx + Axis] != ScaleOrBias.GetData()[DimIdx] && ScaleOrBias.GetData()[DimIdx] != 1)
+					{
+						UE_LOG(LogNNE, Warning, TEXT("LayerNormalization: Scale/bias tensor shape not equal nor broadcastable to input's"));
+						return -1;
+					}
+				}
+
+				return 0;
+			};
+
+			if (const int Res = ValidateScaleOrBias(InputTensors[0]->GetShape(), InputTensors[1]->GetShape()); Res < 0)
+			{
+				return Res;
+			}
+			if (InputTensors.Num() == 3)
+			{
+				if (const int Res = ValidateScaleOrBias(InputTensors[0]->GetShape(), InputTensors[2]->GetShape()); Res < 0)
+				{
+					return Res;
+				}
+			}
+
+			OutputTensors[0]->SetShape(InputTensors[0]->GetShape());
+
+			auto GetMeanOrInvStdDevShape = [Axis = Axis] (const UE::NNE::FTensorShape& Input) -> UE::NNE::FTensorShape
+			{
+				TArray<uint32> OutShape;
+				OutShape.SetNumUninitialized(Input.Rank());
+				
+				for(int DimIdx = 0; DimIdx < Axis; ++DimIdx)
+				{
+					OutShape[DimIdx] = Input.GetData()[DimIdx];
+				}
+				for(int DimIdx = Axis; DimIdx < Input.Rank(); ++DimIdx)
+				{
+					OutShape[DimIdx] = 1;
+				}
+
+				return UE::NNE::FTensorShape::Make(OutShape);
+			};
+
+			if (OutputTensors.Num() >= 2)
+			{
+				OutputTensors[1]->SetShape(GetMeanOrInvStdDevShape(InputTensors[0]->GetShape()));
+			}
+			if (OutputTensors.Num() >= 3)
+			{
+				OutputTensors[2]->SetShape(GetMeanOrInvStdDevShape(InputTensors[0]->GetShape()));
+			}
+
+			return 0;
+		};
+
+		virtual bool Initialize(TConstArrayView<NNE::FTensorDesc> InputTensorDescs, TConstArrayView<NNE::FTensorDesc> OutputTensorDescs, const NNE::FAttributeMap& Attributes) override
+		{
+			check(InputTensorDescs.Num() >= 2 && InputTensorDescs.Num() <= 3);
+			check(OutputTensorDescs.Num() >= 1 && OutputTensorDescs.Num() <= 3);
+
+			const int32 InputRank = InputTensorDescs[0].GetShape().Rank();
+
+			Epsilon = Attributes.GetValueOrDefault<float>(TEXT("epsilon"), Epsilon);
+			Axis = Attributes.GetValueOrDefault<int32>(TEXT("axis"), Axis);
+			if(Axis < 0)
+			{
+				Axis += InputRank;
+			}
+			StashType = Attributes.GetValueOrDefault<int32>(TEXT("stash_type"), StashType);
+
+			if (InputRank != OutputTensorDescs[0].GetShape().Rank())
+			{
+				UE_LOG(LogNNE, Warning, TEXT("LayerNormalization requires the output to have the same rank as the input."));
+				return false;
+			}
+
+			return true;
+		}
+
+		virtual void Dispatch(FRDGBuilder& GraphBuilder, TConstArrayView<FTensorRDGRef> InputTensors, TConstArrayView<FTensorRDGRef> OutputTensors) override
+		{
+			using namespace UE::NNEHlslShaders::Internal;
+
+			check(InputTensors.Num() >= 2 && InputTensors.Num() <= 3);
+			check(OutputTensors.Num() >= 1 && OutputTensors.Num() <= 3);
+
+			check(InputTensors[0] != nullptr);
+			check(InputTensors[1] != nullptr);
+			if(InputTensors.Num() == 3)
+			{
+				check(InputTensors[2] != nullptr);
+			}
+
+			check(OutputTensors[0] != nullptr);
+			if(OutputTensors.Num() >= 2)
+			{
+				check(OutputTensors[1] != nullptr);
+			}
+			if(OutputTensors.Num() >= 3)
+			{
+				check(OutputTensors[2] != nullptr);
+			}
+
+			bool bHasBias = InputTensors.Num() == 3;
+			bool bWriteMean = OutputTensors.Num() >= 2;
+			bool bWriteInvStdDev = OutputTensors.Num() >= 3;
+
+			RDG_EVENT_SCOPE(GraphBuilder, "NNE.Operator.Hlsl.LayerNormalization");
+			RDG_GPU_STAT_SCOPE(GraphBuilder, FNNEOperatorLayerNormalization);
+
+			TLayerNormalizationCS::FParameters* LayerNormParameters = GraphBuilder.AllocParameters<TLayerNormalizationCS::FParameters>();
+			
+			const FTensorRDG& Input = *InputTensors[0];
+			const NNE::FTensorShape& InputShape = Input.GetShape();
+
+			TLayerNormalizationCS::FillInParameters(InputShape.GetData(), Axis, Epsilon, LayerNormParameters);
+
+			const uint32 NumElements = Input.GetVolume();
+			const FIntVector ThreadGroupCount = {
+				(int32) (NumElements / LayerNormParameters->LayerSize),
+				1,
+				1
+			};
+
+			FillTensorStrideShaderParameters(Input, LayerNormParameters->InputTensorInfo, /*Idx*/ 0);
+			LayerNormParameters->Input = GraphBuilder.CreateSRV(FRDGBufferSRVDesc(Input.GetBuffer(), PF_R32_FLOAT));
+
+			const FTensorRDG& Scale = *InputTensors[1];
+			FillTensorStrideShaderParameters(Scale, LayerNormParameters->ScaleTensorInfo, /*Idx*/ 0, /*TargetNumDimsForBroadcast*/ InputShape.Rank() - Axis);
+			LayerNormParameters->InputScale = GraphBuilder.CreateSRV(FRDGBufferSRVDesc(Scale.GetBuffer(), PF_R32_FLOAT));
+
+			if(bHasBias)
+			{
+				const FTensorRDG& Bias = *InputTensors[2];
+				FillTensorStrideShaderParameters(Bias, LayerNormParameters->BiasTensorInfo, /*Idx*/ 0, /*TargetNumDimsForBroadcast*/ InputShape.Rank() - Axis);
+				LayerNormParameters->InputBias = GraphBuilder.CreateSRV(FRDGBufferSRVDesc(Bias.GetBuffer(), PF_R32_FLOAT));
+			}
+
+			const FTensorRDG& Output = *OutputTensors[0];
+			LayerNormParameters->Output = GraphBuilder.CreateUAV(FRDGBufferUAVDesc(Output.GetBuffer(), PF_R32_FLOAT));
+			if(bWriteMean)
+			{
+				const FTensorRDG& OutputMean = *OutputTensors[1];
+				LayerNormParameters->OutputMean = GraphBuilder.CreateUAV(FRDGBufferUAVDesc(OutputMean.GetBuffer(), PF_R32_FLOAT));
+			}
+			if(bWriteInvStdDev)
+			{
+				const FTensorRDG& OutputInvStdDev = *OutputTensors[2];
+				LayerNormParameters->OutputInvStdDev = GraphBuilder.CreateUAV(FRDGBufferUAVDesc(OutputInvStdDev.GetBuffer(), PF_R32_FLOAT));
+			}
+
+			TLayerNormalizationCS::FPermutationDomain PermutationVector;
+			PermutationVector.Set<TLayerNormalizationCS::FLayerNormalizationNumDimensions>(InputShape.Rank());
+			PermutationVector.Set<TLayerNormalizationCS::FLayerNormalizationHasB>(bHasBias);
+			PermutationVector.Set<TLayerNormalizationCS::FLayerNormalizationWriteMean>(bWriteMean);
+			PermutationVector.Set<TLayerNormalizationCS::FLayerNormalizationWriteInvStdDev>(bWriteInvStdDev);
+
+			TShaderMapRef<TLayerNormalizationCS> ComputeShader(GetGlobalShaderMap(GMaxRHIFeatureLevel), PermutationVector);
+
+			FComputeShaderUtils::AddPass(
+				GraphBuilder,
+				RDG_EVENT_NAME("NNE.Operator.Hlsl.LayerNormalization.Dispatch"),
+				ERDGPassFlags::Compute | ERDGPassFlags::NeverCull,
+				ComputeShader,
+				LayerNormParameters,
+				ThreadGroupCount);
+		}
+	};
+
+	bool ValidateLayerNormalizationOperator(const NNE::FAttributeMap& AttributeMap, TConstArrayView<ENNETensorDataType> InputTypes, TConstArrayView<NNE::FSymbolicTensorShape> InputShapes)
+	{
+		bool bIsValid = true;
+
+		FAttributeValidator AttributeValidator;
+		AttributeValidator.AddOptional(TEXT("axis"), ENNEAttributeDataType::Int32);
+		AttributeValidator.AddOptional(TEXT("epsilon"), ENNEAttributeDataType::Float);
+		AttributeValidator.AddOptional(TEXT("stash_type"), ENNEAttributeDataType::Int32);
+		bIsValid &= AttributeValidator.Validate(AttributeMap);
+
+		FInputValidator InputValidator;
+		InputValidator.AddSupportedType(ENNETensorDataType::Float);
+		InputValidator.AddRequired();
+		InputValidator.AddRequired();
+		InputValidator.AddOptional();
+		bIsValid &= InputValidator.Validate(InputTypes);
+
+		return bIsValid;
+	}
+
+	FOperatorHlsl* CreateLayerNormalizationOperator()
+	{
+		return new TLayerNormalization();
+	}
+
+	bool RegisterLayerNormalizationOperator(FOperatorRegistryHlsl& Registry)
+	{
+		// Note: support of a particular version is partial with respect to tensor data types (only the most typical ones are usually supported).
+		Registry.OpAdd({{TEXT("LayerNormalization"), TEXT("Onnx")}, 17}, CreateLayerNormalizationOperator, ValidateLayerNormalizationOperator);
+		return true;
+	}
+} // UE::NNERuntimeRDG::Private::Hlsl
diff --git a/Engine/Plugins/Experimental/NNERuntimeRDG/Source/NNERuntimeRDG/Private/Hlsl/NNERuntimeRDGLayerNormalization.h b/Engine/Plugins/Experimental/NNERuntimeRDG/Source/NNERuntimeRDG/Private/Hlsl/NNERuntimeRDGLayerNormalization.h
new file mode 100644
--- /dev/null
+++ b/Engine/Plugins/Experimental/NNERuntimeRDG/Source/NNERuntimeRDG/Private/Hlsl/NNERuntimeRDGLayerNormalization.h
@@ -0,0 +1,11 @@
+// Copyright Epic Games, Inc. All Rights Reserved.
+
+#pragma once
+
+#include "NNERuntimeRDGBase.h"
+#include "NNERuntimeRDGHlslOp.h"
+
+namespace UE::NNERuntimeRDG::Private::Hlsl
+{
+	bool RegisterLayerNormalizationOperator(FOperatorRegistryHlsl& Registry);
+} // UE::NNERuntimeRDG::Private::Hlsl
diff --git a/Engine/Plugins/Experimental/NNERuntimeRDG/Source/NNERuntimeRDG/Private/NNERuntimeRDGHlsl.cpp b/Engine/Plugins/Experimental/NNERuntimeRDG/Source/NNERuntimeRDG/Private/NNERuntimeRDGHlsl.cpp
--- a/Engine/Plugins/Experimental/NNERuntimeRDG/Source/NNERuntimeRDG/Private/NNERuntimeRDGHlsl.cpp
+++ b/Engine/Plugins/Experimental/NNERuntimeRDG/Source/NNERuntimeRDG/Private/NNERuntimeRDGHlsl.cpp
@@ -31,6 +31,7 @@
 #include "Hlsl/NNERuntimeRDGGlobalPool.h"
 #include "Hlsl/NNERuntimeRDGIdentity.h"
 #include "Hlsl/NNERuntimeRDGInstanceNormalization.h"
+#include "Hlsl/NNERuntimeRDGLayerNormalization.h"
 #include "Hlsl/NNERuntimeRDGPad.h"
 #include "Hlsl/NNERuntimeRDGPool.h"
 #include "Hlsl/NNERuntimeRDGReduce.h"
@@ -71,6 +72,7 @@ bool UNNERuntimeRDGHlslImpl::Init()
 	RegisterGlobalPoolOperators(*Registry);
 	RegisterIdentityOperator(*Registry);
 	RegisterInstanceNormalizationOperator(*Registry);
+	RegisterLayerNormalizationOperator(*Registry);
 	RegisterPadOperator(*Registry);
 	RegisterPoolOperators(*Registry);
 	RegisterReduceOperators(*Registry);
