NNE - Change Hlsl LayerNormalization implementation

Problems:
* Hlsl LayerNormalization's shader triggers some false race condition detections when compiled with d3dcompiler (when building standalone), although the code is correct and works correctly.

Solutions:
* Rewrite Hlsl LayerNormalization to split Mean/InvStdDev calculation and normalization in two separate passes

Tests:
* Win64 Standalone tests: LayerNormalization
* Win64 Editor tests: LayerNormalization

#rnx

diff --git a/Engine/Plugins/Experimental/NNERuntimeRDG/Source/NNEHlslShaders/Internal/NNEHlslShadersLayerNormalizationCS.h b/Engine/Plugins/Experimental/NNERuntimeRDG/Source/NNEHlslShaders/Internal/NNEHlslShadersLayerNormalizationCS.h
--- a/Engine/Plugins/Experimental/NNERuntimeRDG/Source/NNEHlslShaders/Internal/NNEHlslShadersLayerNormalizationCS.h
+++ b/Engine/Plugins/Experimental/NNERuntimeRDG/Source/NNEHlslShaders/Internal/NNEHlslShadersLayerNormalizationCS.h
@@ -24,14 +24,14 @@ namespace UE::NNEHlslShaders::Internal
 
 		class FLayerNormalizationNumDimensions : SHADER_PERMUTATION_RANGE_INT("NUM_DIMENSIONS", 1, FLayerNormalizationConstants::MAX_NUM_DIMENSIONS);
 		class FLayerNormalizationHasB : SHADER_PERMUTATION_BOOL("HAS_B");
-		class FLayerNormalizationWriteMean : SHADER_PERMUTATION_BOOL("WRITE_MEAN");
-		class FLayerNormalizationWriteInvStdDev : SHADER_PERMUTATION_BOOL("WRITE_INVSTDDEV");
-		using FPermutationDomain = TShaderPermutationDomain<FLayerNormalizationNumDimensions, FLayerNormalizationHasB, FLayerNormalizationWriteMean, FLayerNormalizationWriteInvStdDev>;
+		using FPermutationDomain = TShaderPermutationDomain<FLayerNormalizationNumDimensions, FLayerNormalizationHasB>;
 
 	public:
 
 		BEGIN_SHADER_PARAMETER_STRUCT(FParameters, )
+			SHADER_PARAMETER(uint32, Num)
 			SHADER_PARAMETER(uint32, Axis)
+			SHADER_PARAMETER(uint32, ThreadCountX)
 			SHADER_PARAMETER(float, Epsilon)
 			SHADER_PARAMETER(uint32, LayerSize)
 			SHADER_PARAMETER_RDG_BUFFER_SRV(Buffer<float>, Input)
@@ -40,9 +40,9 @@ namespace UE::NNEHlslShaders::Internal
 			SHADER_PARAMETER_ARRAY(FUintVector4, ScaleTensorInfo, [FLayerNormalizationConstants::MAX_NUM_DIMENSIONS])
 			SHADER_PARAMETER_RDG_BUFFER_SRV(Buffer<float>, InputBias)
 			SHADER_PARAMETER_ARRAY(FUintVector4, BiasTensorInfo, [FLayerNormalizationConstants::MAX_NUM_DIMENSIONS])
+			SHADER_PARAMETER_RDG_BUFFER_SRV(Buffer<float>, InputMean)
+			SHADER_PARAMETER_RDG_BUFFER_SRV(Buffer<float>, InputInvStdDev)
 			SHADER_PARAMETER_RDG_BUFFER_UAV(RWBuffer<float>, Output)
-			SHADER_PARAMETER_RDG_BUFFER_UAV(RWBuffer<float>, OutputMean)
-			SHADER_PARAMETER_RDG_BUFFER_UAV(RWBuffer<float>, OutputInvStdDev)
 			END_SHADER_PARAMETER_STRUCT()
 
 			static void FillInParameters(TConstArrayView<uint32> Shape, int32 Axis, float Epsilon, FParameters* Parameters);
diff --git a/Engine/Plugins/Experimental/NNERuntimeRDG/Source/NNEHlslShaders/Shaders/NNEHlslShadersLayerNormalization.usf b/Engine/Plugins/Experimental/NNERuntimeRDG/Source/NNEHlslShaders/Shaders/NNEHlslShadersLayerNormalization.usf
--- a/Engine/Plugins/Experimental/NNERuntimeRDG/Source/NNEHlslShaders/Shaders/NNEHlslShadersLayerNormalization.usf
+++ b/Engine/Plugins/Experimental/NNERuntimeRDG/Source/NNEHlslShaders/Shaders/NNEHlslShadersLayerNormalization.usf
@@ -3,7 +3,6 @@
 #include "/Engine/Public/Platform.ush"
 #include "NNEHlslShadersBroadcastHelper.ush"
 
-//NOTE: once have support for more datatypes, make BUFFER_TYPE into a variant and WORK_TYPE depend on stash_type
 #define WORK_TYPE float
 #define BUFFER_TYPE float
 #define READ(x) x
@@ -15,12 +14,15 @@ Buffer<BUFFER_TYPE> InputScale;
 uint4 ScaleTensorInfo[NUM_DIMENSIONS];
 Buffer<BUFFER_TYPE> InputBias;
 uint4 BiasTensorInfo[NUM_DIMENSIONS];
+Buffer<WORK_TYPE> InputMean;
+Buffer<WORK_TYPE> InputInvStdDev;
 RWBuffer<BUFFER_TYPE> Output;
-RWBuffer<BUFFER_TYPE> OutputMean;
-RWBuffer<BUFFER_TYPE> OutputInvStdDev;
-uint Axis;
+
 float Epsilon;
+int Axis;
+uint ThreadCountX;
 uint LayerSize;
+uint Num;
 
 #define STRIDE_IDX 0
 
@@ -49,78 +51,17 @@ uint GetInnerIdx(const uint DimIterator[NUM_DIMENSIONS], const uint4 TensorInfo[
 	return Offset;
 }
 
-groupshared WORK_TYPE SharedMemoryMean[THREADGROUP_SIZE];
-groupshared WORK_TYPE SharedMemoryMean2[THREADGROUP_SIZE];
-groupshared WORK_TYPE SharedMemoryCount[THREADGROUP_SIZE];
 
 [numthreads(THREADGROUP_SIZE, 1, 1)]
-void LayerNormalization(in const uint3 GroupID : SV_GroupID, in const uint3 GroupThreadID : SV_GroupThreadID, in const uint3 DispatchThreadID : SV_DispatchThreadID)
+void LayerNormalization(in const uint3 DispatchThreadID : SV_DispatchThreadID)
 {
-	const int SliceIdx = GroupID.x; // each thread group processes a slice
-	const int InputStartIdx = SliceIdx * LayerSize;
-	
-
-	// Online Welford using sharedmem and THREADGROUP_SIZE threads along the dimensions after `Axis` (included)
-	{
-		WORK_TYPE Mean = 0;
-		WORK_TYPE Mean2 = 0;
-		int Count = 0;
-
-		for (int InElemIdx = GroupThreadID.x; InElemIdx < LayerSize; InElemIdx += THREADGROUP_SIZE)
-		{
-			WORK_TYPE Value = (WORK_TYPE) READ(Input[InElemIdx + InputStartIdx]);
-		
-			Count += 1;
-			WORK_TYPE Delta = Value - Mean;
-			Mean += Delta / (WORK_TYPE)Count;
-			WORK_TYPE Delta2 = Value - Mean;
-			Mean2 += Delta * Delta2;
-		}
-
-		SharedMemoryMean[GroupThreadID.x] = Mean;
-		SharedMemoryMean2[GroupThreadID.x] = Mean2;
-		SharedMemoryCount[GroupThreadID.x] = Count;
-	}
-
-	GroupMemoryBarrierWithGroupSync();
-
-	if (GroupThreadID.x == 0)
+	const uint ElementIdx = DispatchThreadID.y * ThreadCountX + DispatchThreadID.x;
+	if(ElementIdx < Num)
 	{
-		WORK_TYPE Mean = 0;
-		WORK_TYPE Mean2 = 0;
-		int Count = 0;
-
-		UNROLL
-		for (int Idx = 0; Idx < THREADGROUP_SIZE; Idx++)
-		{
-			WORK_TYPE CurrMean = SharedMemoryMean[Idx];
-			WORK_TYPE CurrMean2 = SharedMemoryMean2[Idx];
-			int CurrCount = SharedMemoryCount[Idx];
-			int NewCount = Count + CurrCount;
-			WORK_TYPE Delta = CurrMean - Mean;
-			Mean = (CurrCount * CurrMean + Count * Mean) / NewCount;
-			Mean2 += CurrMean2 + Delta * Delta * (WORK_TYPE)Count * (WORK_TYPE)CurrCount / (WORK_TYPE)NewCount;
-			Count = NewCount;
-		}
-
-		SharedMemoryMean[0] = Mean;
-		SharedMemoryMean2[0] = Mean2;
-		SharedMemoryCount[0] = Count;
+		const uint LayerIdx = ElementIdx / LayerSize;
 		
-	}
-
-	GroupMemoryBarrierWithGroupSync();
-
-	WORK_TYPE Mean = SharedMemoryMean[0];
-	WORK_TYPE Mean2 = SharedMemoryMean2[0];
-	int Count = SharedMemoryCount[0];
-
-	WORK_TYPE Variance = Mean2 / (WORK_TYPE)Count;
-	WORK_TYPE InvStdDev = (WORK_TYPE)1 / sqrt(Variance + Epsilon);
-
-	for (int InElemIdx = GroupThreadID.x; InElemIdx < LayerSize; InElemIdx += THREADGROUP_SIZE)
-	{
-		const int ElementIdx = InElemIdx + InputStartIdx;
+		WORK_TYPE Mean = InputMean[LayerIdx];
+		WORK_TYPE InvStdDev = InputInvStdDev[LayerIdx];
 
 		BUFFER_TYPE Normalized = (BUFFER_TYPE) ( (READ(Input[ElementIdx]) - Mean) * InvStdDev );
 
@@ -128,17 +69,10 @@ void LayerNormalization(in const uint3 GroupID : SV_GroupID, in const uint3 Grou
 		GetDimIterator(ElementIdx, InputTensorInfo, InputDimIterator);
 
 		BUFFER_TYPE NormalizedScaled = Normalized * READ( InputScale[GetInnerIdx(InputDimIterator, ScaleTensorInfo, Axis)]);
-#if HAS_B == 0
-		Output[ElementIdx] = WRITE( NormalizedScaled );
-#else
-		Output[ElementIdx] = WRITE( NormalizedScaled + READ( InputBias[GetInnerIdx(InputDimIterator, BiasTensorInfo, Axis)] ) );
-#endif
-
+		#if HAS_B == 0
+			Output[ElementIdx] = WRITE( NormalizedScaled );
+		#else
+			Output[ElementIdx] = WRITE( NormalizedScaled + READ( InputBias[GetInnerIdx(InputDimIterator, BiasTensorInfo, Axis)] ) );
+		#endif
 	}
-#if WRITE_MEAN != 0
-	OutputMean[SliceIdx] = WRITE( (BUFFER_TYPE) Mean );
-#endif
-#if WRITE_INVSTDDEV != 0
-	OutputInvStdDev[SliceIdx] = WRITE( (BUFFER_TYPE) InvStdDev );
-#endif
 }
diff --git a/Engine/Plugins/Experimental/NNERuntimeRDG/Source/NNERuntimeRDG/Private/Hlsl/NNERuntimeRDGLayerNormalization.cpp b/Engine/Plugins/Experimental/NNERuntimeRDG/Source/NNERuntimeRDG/Private/Hlsl/NNERuntimeRDGLayerNormalization.cpp
--- a/Engine/Plugins/Experimental/NNERuntimeRDG/Source/NNERuntimeRDG/Private/Hlsl/NNERuntimeRDGLayerNormalization.cpp
+++ b/Engine/Plugins/Experimental/NNERuntimeRDG/Source/NNERuntimeRDG/Private/Hlsl/NNERuntimeRDGLayerNormalization.cpp
@@ -2,6 +2,7 @@
 
 #include "NNERuntimeRDGLayerNormalization.h"
 #include "NNEHlslShadersLayerNormalizationCS.h"
+#include "NNEHlslShadersReduceCS.h"
 #include "NNERuntimeRDGHlslHelper.h"
 #include "NNETensor.h"
 #include "NNETypes.h"
@@ -150,19 +151,64 @@ namespace UE::NNERuntimeRDG::Private::Hlsl
 			RDG_EVENT_SCOPE(GraphBuilder, "NNE.Operator.Hlsl.LayerNormalization");
 			RDG_GPU_STAT_SCOPE(GraphBuilder, FNNEOperatorLayerNormalization);
 
-			TLayerNormalizationCS::FParameters* LayerNormParameters = GraphBuilder.AllocParameters<TLayerNormalizationCS::FParameters>();
-			
 			const FTensorRDG& Input = *InputTensors[0];
 			const NNE::FTensorShape& InputShape = Input.GetShape();
+			
+			// First apply Reduction() to temp buffers getting Mean and InvStdDev
+			TReduceCS::FParameters* ReduceParameters = GraphBuilder.AllocParameters<TReduceCS::FParameters>();
+			// Need to feed a shape such that the last dimension corresponds to the slice to reduce
+			TArray<uint32> ReductionShape;
+			ReductionShape.Init(1, Axis + 1);
+			for(int Idx = 0; Idx < InputShape.Rank(); ++Idx)
+			{
+				if(Idx < Axis)
+				{
+					ReductionShape[Idx] = InputShape.GetData()[Idx];
+				}
+				else
+				{
+					ReductionShape[Axis] *= InputShape.GetData()[Idx];
+				}
+			}
+
+			TReduceCS::FillInParameters(ReductionShape, Axis, ReduceParameters);
+			ReduceParameters->Epsilon = Epsilon;
+			// NOTE: once have support for more datatypes, make this depend on stash_type attribute
+			uint32 BytesPerElementTemp = Input.GetElementByteSize(); 
+			const FRDGBufferDesc LayerNormTempBufferDesc = FRDGBufferDesc::CreateBufferDesc(BytesPerElementTemp, ReduceParameters->NumElemBeforeAxis);
+
+			FRDGBufferRef MeanBuffer;
+			if(bWriteMean)
+			{
+				const FTensorRDG& OutputMean = *OutputTensors[1];
+				MeanBuffer = OutputMean.GetBuffer();
+			}
+			else
+			{
+				MeanBuffer = GraphBuilder.CreateBuffer(LayerNormTempBufferDesc, TEXT("NNE.Operator.Hlsl.LayerNormalization.TempMeanBuffer"), ERDGBufferFlags::None);
+			}
+
+			FRDGBufferRef InvStdDevBuffer;
+			if(bWriteInvStdDev)
+			{
+				const FTensorRDG& OutputInvStdDev = *OutputTensors[2];
+				InvStdDevBuffer = OutputInvStdDev.GetBuffer();
+			}
+			else
+			{
+				InvStdDevBuffer = GraphBuilder.CreateBuffer(LayerNormTempBufferDesc, TEXT("NNE.Operator.Hlsl.LayerNormalization.TempInvStdDevBuffer"), ERDGBufferFlags::None);
+			}
+
+			TReduceCS::EnqueueRDG(GraphBuilder, ReduceParameters, Input.GetBuffer(), MeanBuffer, EReduceOperatorType::AverageInvStdDev, InvStdDevBuffer);
 
+			// Then LayerNormalization
+			TLayerNormalizationCS::FParameters* LayerNormParameters = GraphBuilder.AllocParameters<TLayerNormalizationCS::FParameters>();
 			TLayerNormalizationCS::FillInParameters(InputShape.GetData(), Axis, Epsilon, LayerNormParameters);
 
 			const uint32 NumElements = Input.GetVolume();
-			const FIntVector ThreadGroupCount = {
-				(int32) (NumElements / LayerNormParameters->LayerSize),
-				1,
-				1
-			};
+			LayerNormParameters->Num = NumElements;
+			const FIntVector ThreadGroupCount = ComputeElementWiseThreadGroups(NumElements, FLayerNormalizationConstants::NUM_GROUP_THREADS);;
+			LayerNormParameters->ThreadCountX = ThreadGroupCount.X * FLayerNormalizationConstants::NUM_GROUP_THREADS;
 
 			FillTensorStrideShaderParameters(Input, LayerNormParameters->InputTensorInfo, /*Idx*/ 0);
 			LayerNormParameters->Input = GraphBuilder.CreateSRV(FRDGBufferSRVDesc(Input.GetBuffer(), PF_R32_FLOAT));
@@ -178,24 +224,15 @@ namespace UE::NNERuntimeRDG::Private::Hlsl
 				LayerNormParameters->InputBias = GraphBuilder.CreateSRV(FRDGBufferSRVDesc(Bias.GetBuffer(), PF_R32_FLOAT));
 			}
 
+			LayerNormParameters->InputMean = GraphBuilder.CreateSRV(FRDGBufferSRVDesc(MeanBuffer, PF_R32_FLOAT));
+			LayerNormParameters->InputInvStdDev = GraphBuilder.CreateSRV(FRDGBufferSRVDesc(InvStdDevBuffer, PF_R32_FLOAT));
+
 			const FTensorRDG& Output = *OutputTensors[0];
 			LayerNormParameters->Output = GraphBuilder.CreateUAV(FRDGBufferUAVDesc(Output.GetBuffer(), PF_R32_FLOAT));
-			if(bWriteMean)
-			{
-				const FTensorRDG& OutputMean = *OutputTensors[1];
-				LayerNormParameters->OutputMean = GraphBuilder.CreateUAV(FRDGBufferUAVDesc(OutputMean.GetBuffer(), PF_R32_FLOAT));
-			}
-			if(bWriteInvStdDev)
-			{
-				const FTensorRDG& OutputInvStdDev = *OutputTensors[2];
-				LayerNormParameters->OutputInvStdDev = GraphBuilder.CreateUAV(FRDGBufferUAVDesc(OutputInvStdDev.GetBuffer(), PF_R32_FLOAT));
-			}
 
 			TLayerNormalizationCS::FPermutationDomain PermutationVector;
 			PermutationVector.Set<TLayerNormalizationCS::FLayerNormalizationNumDimensions>(InputShape.Rank());
 			PermutationVector.Set<TLayerNormalizationCS::FLayerNormalizationHasB>(bHasBias);
-			PermutationVector.Set<TLayerNormalizationCS::FLayerNormalizationWriteMean>(bWriteMean);
-			PermutationVector.Set<TLayerNormalizationCS::FLayerNormalizationWriteInvStdDev>(bWriteInvStdDev);
 
 			TShaderMapRef<TLayerNormalizationCS> ComputeShader(GetGlobalShaderMap(GMaxRHIFeatureLevel), PermutationVector);
 
