NNE - Add support for Clip 11, 12, 13

Problems:
* RDG Hlsl backend misses support for Clip-11, Clip-12, Clip-13 and relative tests.

Solutions:
* Add support for said versions of Clip.

Tests:
* Win64 Standalone tests: Clip-6, Clip-11, Clip-12, Clip-13
* Win64 Editor tests: Clip-6, Clip-11, Clip-12, Clip-13

#rnx

#virtualized

diff --git a/Engine/Plugins/Experimental/NNERuntimeRDG/Source/NNEHlslShaders/Internal/NNEHlslShadersElementWiseUnaryCS.h b/Engine/Plugins/Experimental/NNERuntimeRDG/Source/NNEHlslShaders/Internal/NNEHlslShadersElementWiseUnaryCS.h
--- a/Engine/Plugins/Experimental/NNERuntimeRDG/Source/NNEHlslShaders/Internal/NNEHlslShadersElementWiseUnaryCS.h
+++ b/Engine/Plugins/Experimental/NNERuntimeRDG/Source/NNEHlslShaders/Internal/NNEHlslShadersElementWiseUnaryCS.h
@@ -22,7 +22,9 @@ namespace UE::NNEHlslShaders::Internal
 		SHADER_USE_PARAMETER_STRUCT(TElementWiseUnaryCS, FHlslShaderBase)
 
 		class FOperatorType : SHADER_PERMUTATION_ENUM_CLASS("OP_TYPENAME", NNE::Internal::EElementWiseUnaryOperatorType);
-		using FPermutationDomain = TShaderPermutationDomain<FOperatorType>;
+		class FAlphaOnGPU : SHADER_PERMUTATION_BOOL("ALPHA_ON_GPU");
+		class FBetaOnGPU : SHADER_PERMUTATION_BOOL("BETA_ON_GPU");
+		using FPermutationDomain = TShaderPermutationDomain<FOperatorType,FAlphaOnGPU,FBetaOnGPU>;
 
 	public:
 
@@ -32,10 +34,14 @@ namespace UE::NNEHlslShaders::Internal
 			SHADER_PARAMETER(uint32, Num)
 			SHADER_PARAMETER(uint32, ThreadCountX)
 			SHADER_PARAMETER(float, Alpha)
+			SHADER_PARAMETER_RDG_BUFFER_SRV(Buffer<float>, AlphaTensor)
 			SHADER_PARAMETER(float, Beta)
+			SHADER_PARAMETER_RDG_BUFFER_SRV(Buffer<float>, BetaTensor)
 			SHADER_PARAMETER(float, Gamma)
 		END_SHADER_PARAMETER_STRUCT()
 
+		static bool ShouldCompilePermutation(const FGlobalShaderPermutationParameters& Parameters);
+
 		static void ModifyCompilationEnvironment(const FGlobalShaderPermutationParameters& InParameters, FShaderCompilerEnvironment& OutEnvironment);
 
 	private:
diff --git a/Engine/Plugins/Experimental/NNERuntimeRDG/Source/NNEHlslShaders/Private/NNEHlslShadersElementWiseUnaryCS.cpp b/Engine/Plugins/Experimental/NNERuntimeRDG/Source/NNEHlslShaders/Private/NNEHlslShadersElementWiseUnaryCS.cpp
--- a/Engine/Plugins/Experimental/NNERuntimeRDG/Source/NNEHlslShaders/Private/NNEHlslShadersElementWiseUnaryCS.cpp
+++ b/Engine/Plugins/Experimental/NNERuntimeRDG/Source/NNEHlslShaders/Private/NNEHlslShadersElementWiseUnaryCS.cpp
@@ -5,6 +5,22 @@
 
 namespace UE::NNEHlslShaders::Internal
 {
+	bool TElementWiseUnaryCS::ShouldCompilePermutation(const FGlobalShaderPermutationParameters& Parameters)
+	{
+		if (!FHlslShaderBase::ShouldCompilePermutation(Parameters))
+		{
+			return false;
+		}
+
+		FPermutationDomain PermutationVector(Parameters.PermutationId);
+
+		NNE::Internal::EElementWiseUnaryOperatorType OpType = PermutationVector.Get<TElementWiseUnaryCS::FOperatorType>();
+		const bool bAlphaOnGpu = PermutationVector.Get<TElementWiseUnaryCS::FAlphaOnGPU>();
+		const bool bBetaOnGpu = PermutationVector.Get<TElementWiseUnaryCS::FBetaOnGPU>();
+
+		return OpType == NNE::Internal::EElementWiseUnaryOperatorType::Clip || (!bAlphaOnGpu && !bBetaOnGpu);
+	}
+
 	void TElementWiseUnaryCS::ModifyCompilationEnvironment(const FGlobalShaderPermutationParameters& InParameters, FShaderCompilerEnvironment& OutEnvironment)
 	{
 		FGlobalShader::ModifyCompilationEnvironment(InParameters, OutEnvironment);
diff --git a/Engine/Plugins/Experimental/NNERuntimeRDG/Source/NNEHlslShaders/Shaders/NNEHlslShadersElementWiseUnary.usf b/Engine/Plugins/Experimental/NNERuntimeRDG/Source/NNEHlslShaders/Shaders/NNEHlslShadersElementWiseUnary.usf
--- a/Engine/Plugins/Experimental/NNERuntimeRDG/Source/NNEHlslShaders/Shaders/NNEHlslShadersElementWiseUnary.usf
+++ b/Engine/Plugins/Experimental/NNERuntimeRDG/Source/NNEHlslShaders/Shaders/NNEHlslShadersElementWiseUnary.usf
@@ -7,9 +7,29 @@ RWBuffer<float> Output;
 uint Num;
 uint ThreadCountX;
 float Alpha;
+Buffer<float> AlphaTensor;
 float Beta;
+Buffer<float> BetaTensor;
 float Gamma;
 
+float GetAlpha()
+{
+#if ALPHA_ON_GPU == 0
+	return Alpha;
+#else
+	return AlphaTensor[0];
+#endif
+}
+
+float GetBeta()
+{
+#if BETA_ON_GPU == 0
+	return Beta;
+#else
+	return BetaTensor[0];
+#endif
+}
+
 //see [REDACTED_PROFILE_URL]/onnx/blob/main/docs/Operators.md#Relu
 float relu(float x)
 {
@@ -19,7 +39,7 @@ float relu(float x)
 //see [REDACTED_PROFILE_URL]/onnx/blob/main/docs/Operators.md#Selu
 float selu(float x)
 {
-	float yNegOrZero = Gamma * (Alpha * exp(x) - Alpha);
+	float yNegOrZero = Gamma * (GetAlpha() * exp(x) - GetAlpha());
 	float yPos = Gamma * x;
 	return (x > 0.0f) ? yPos : yNegOrZero;
 }
@@ -68,7 +88,7 @@ float atanh(float x)
 //see [REDACTED_PROFILE_URL]/onnx/blob/main/docs/Operators.md#elu
 float elu(float x)
 {
-    float yNeg = Alpha * (exp(x) - 1.0f);
+    float yNeg = GetAlpha() * (exp(x) - 1.0f);
 	float yPosOrZero = x;
 	return (x >= 0.0f) ? yPosOrZero : yNeg;
 }
@@ -91,7 +111,7 @@ float hardSigmoid(float x, float alpha, float beta)
 //see [REDACTED_PROFILE_URL]/onnx/blob/main/docs/Operators.md#hardSigmoid
 float hardSigmoid(float x)
 {
-	return hardSigmoid(x, Alpha, Beta);
+	return hardSigmoid(x, GetAlpha(), GetBeta());
 }
 
 //see [REDACTED_PROFILE_URL]/onnx/blob/main/docs/Operators.md#hardSwish
@@ -103,16 +123,16 @@ float hardSwish(float x)
 //see [REDACTED_PROFILE_URL]/onnx/blob/main/docs/Operators.md#leakyRelu
 float leakyRelu(float x)
 {
-    return (x >= 0.0f) ? x : Alpha * x;
+    return (x >= 0.0f) ? x : GetAlpha() * x;
 }
 
 
 //see [REDACTED_PROFILE_URL]/onnx/blob/main/docs/Operators.md#clip
 float clipOp(float x)
 {
-	float min = Alpha;
-	float max = Beta;
-	return clamp(x, Alpha, Beta);
+	float min = GetAlpha();
+	float max = GetBeta();
+	return clamp(x, GetAlpha(), GetBeta());
 }
 
 
diff --git a/Engine/Plugins/Experimental/NNERuntimeRDG/Source/NNERuntimeRDG/Private/Hlsl/NNERuntimeRDGElementWiseUnary.cpp b/Engine/Plugins/Experimental/NNERuntimeRDG/Source/NNERuntimeRDG/Private/Hlsl/NNERuntimeRDGElementWiseUnary.cpp
--- a/Engine/Plugins/Experimental/NNERuntimeRDG/Source/NNERuntimeRDG/Private/Hlsl/NNERuntimeRDGElementWiseUnary.cpp
+++ b/Engine/Plugins/Experimental/NNERuntimeRDG/Source/NNERuntimeRDG/Private/Hlsl/NNERuntimeRDGElementWiseUnary.cpp
@@ -23,7 +23,7 @@ namespace UE::NNERuntimeRDG::Private::Hlsl
 	{
 	public:
 
-		TElementWiseUnary() {}
+		TElementWiseUnary(int InVersion) : Version(InVersion) {};
 		virtual ~TElementWiseUnary() = default;
 
 	private:
@@ -32,16 +32,39 @@ namespace UE::NNERuntimeRDG::Private::Hlsl
 		float Beta = 0.0f;
 		float Gamma = 0.0f;
 
+		int Version = 0;
+
 	public:
 
 		virtual int PrepareOutputs(TConstArrayView<NNE::Internal::FTensorRef> InputTensors, TArrayView<NNE::Internal::FTensorRef> OutputTensors) override
 		{
-			check(InputTensors.Num() == 1);
+			check((OpType == NNE::Internal::EElementWiseUnaryOperatorType::Clip && Version >= 11) ? (InputTensors.Num() >= 1 && InputTensors.Num() <= 3) : InputTensors.Num() == 1);
 			check(OutputTensors.Num() == 1);
+			
 			OutputTensors[0]->SetShape(InputTensors[0]->GetShape());
 
 			const NNE::Internal::FTensor& X = *InputTensors[0];
 
+			if(OpType == NNE::Internal::EElementWiseUnaryOperatorType::Clip && Version >= 11 && InputTensors.Num() >= 2)
+			{
+				{
+					const NNE::Internal::FTensor& MinTensor = *InputTensors[1];
+					if(MinTensor.HasPreparedData())
+					{
+						Alpha = MinTensor.GetPreparedData<float>()[0];
+					}
+				}
+
+				if(InputTensors.Num() == 3)
+				{
+					const NNE::Internal::FTensor& MaxTensor = *InputTensors[2];
+					if(MaxTensor.HasPreparedData())
+					{
+						Beta = MaxTensor.GetPreparedData<float>()[0];
+					}
+				}
+			}
+
 			Internal::CPUHelper::ElementWiseUnary::Apply(OpType, X, Alpha, Beta, Gamma, *OutputTensors[0]);
 			
 			return 0;
@@ -81,6 +104,33 @@ namespace UE::NNERuntimeRDG::Private::Hlsl
 			TElementWiseUnaryCS::FPermutationDomain PermutationVector;
 
 			PermutationVector.Set<TElementWiseUnaryCS::FOperatorType>(OpType);
+			PermutationVector.Set<TElementWiseUnaryCS::FAlphaOnGPU>(false);
+			PermutationVector.Set<TElementWiseUnaryCS::FBetaOnGPU>(false);
+
+			if(OpType == NNE::Internal::EElementWiseUnaryOperatorType::Clip && Version >= 11 && InInputTensors.Num() >= 2)
+			{
+				{
+					const NNE::Internal::FTensor& MinTensor = *InInputTensors[1];
+					if(!MinTensor.HasPreparedData())
+					{
+						PermutationVector.Set<TElementWiseUnaryCS::FAlphaOnGPU>(true);
+						FRDGBufferSRVRef AlphaSRV = GraphBuilder.CreateSRV(FRDGBufferSRVDesc(InInputTensors[1]->GetBuffer(), PF_R32_FLOAT));
+						Params->AlphaTensor = AlphaSRV;
+					}
+				}
+
+				if(InInputTensors.Num() == 3)
+				{
+					const NNE::Internal::FTensor& MaxTensor = *InInputTensors[2];
+					if(!MaxTensor.HasPreparedData())
+					{
+						PermutationVector.Set<TElementWiseUnaryCS::FBetaOnGPU>(true);
+						FRDGBufferSRVRef BetaSRV = GraphBuilder.CreateSRV(FRDGBufferSRVDesc(InInputTensors[2]->GetBuffer(), PF_R32_FLOAT));
+						Params->BetaTensor = BetaSRV;
+					}
+				}
+			}
+
 
 			TShaderMapRef<TElementWiseUnaryCS> ComputeShader(GetGlobalShaderMap(GMaxRHIFeatureLevel), PermutationVector);
 
@@ -97,29 +147,29 @@ namespace UE::NNERuntimeRDG::Private::Hlsl
 		}
 	};
 
-	template<> TElementWiseUnary<NNE::Internal::EElementWiseUnaryOperatorType::Selu>::TElementWiseUnary()
-		: Alpha(1.67326319217681884765625f), Beta(0.0f), Gamma(1.05070102214813232421875f)
+	template<> TElementWiseUnary<NNE::Internal::EElementWiseUnaryOperatorType::Selu>::TElementWiseUnary(int InVersion)
+		: Alpha(1.67326319217681884765625f), Beta(0.0f), Gamma(1.05070102214813232421875f), Version(InVersion)
 	{
 	}
 
-	template<> TElementWiseUnary<NNE::Internal::EElementWiseUnaryOperatorType::Elu>::TElementWiseUnary()
-		: Alpha(1.0f), Beta(0.0f), Gamma(0.0f) 
+	template<> TElementWiseUnary<NNE::Internal::EElementWiseUnaryOperatorType::Elu>::TElementWiseUnary(int InVersion)
+		: Alpha(1.0f), Beta(0.0f), Gamma(0.0f), Version(InVersion)
 	{
 	}
 
-	template<> TElementWiseUnary<NNE::Internal::EElementWiseUnaryOperatorType::HardSigmoid>::TElementWiseUnary()
-		: Alpha(0.2f), Beta(0.5f), Gamma(0.0f)
+	template<> TElementWiseUnary<NNE::Internal::EElementWiseUnaryOperatorType::HardSigmoid>::TElementWiseUnary(int InVersion)
+		: Alpha(0.2f), Beta(0.5f), Gamma(0.0f), Version(InVersion)
 	{
 	}
 
-	template<> TElementWiseUnary<NNE::Internal::EElementWiseUnaryOperatorType::LeakyRelu>::TElementWiseUnary()
-		: Alpha(0.01f), Beta(0.0f), Gamma(0.0f)
+	template<> TElementWiseUnary<NNE::Internal::EElementWiseUnaryOperatorType::LeakyRelu>::TElementWiseUnary(int InVersion)
+		: Alpha(0.01f), Beta(0.0f), Gamma(0.0f), Version(InVersion)
 	{
 	}
 
 	template<> bool TElementWiseUnary<NNE::Internal::EElementWiseUnaryOperatorType::Clip>::Initialize(TConstArrayView<NNE::FTensorDesc> InputTensorDescs, TConstArrayView<NNE::FTensorDesc> OutputTensorDescs, const NNE::FAttributeMap& Attributes)
 	{
-		check(InputTensorDescs.Num() == 1);
+		check(Version >= 11 ? (InputTensorDescs.Num() >= 1 && InputTensorDescs.Num() <= 3) : InputTensorDescs.Num() == 1);
 		check(OutputTensorDescs.Num() == 1);
 
 		Alpha = Attributes.GetValueOrDefault(TEXT("min"), -3.402823e+38f);
@@ -127,14 +177,14 @@ namespace UE::NNERuntimeRDG::Private::Hlsl
 		return true;
 	}
 
-	template<NNE::Internal::EElementWiseUnaryOperatorType OpType>
+	template<NNE::Internal::EElementWiseUnaryOperatorType OpType, int Version>
 	FOperatorHlsl* CreateElementWiseUnaryOperator()
 	{
-		return new TElementWiseUnary<OpType>();
+		return new TElementWiseUnary<OpType>(Version);
 	}
-
+	
 	template<NNE::Internal::EElementWiseUnaryOperatorType OpType>
-	bool ValidateElementWiseUnaryOperator(const NNE::FAttributeMap& AttributeMap, TConstArrayView<ENNETensorDataType> InputTypes, TConstArrayView<NNE::FSymbolicTensorShape> InputShapes)
+	bool ValidateElementWiseUnaryOperatorImpl(int Version, const NNE::FAttributeMap& AttributeMap, TConstArrayView<ENNETensorDataType> InputTypes, TConstArrayView<NNE::FSymbolicTensorShape> InputShapes)
 	{
 		bool bIsValid = true;
 
@@ -150,7 +200,7 @@ namespace UE::NNERuntimeRDG::Private::Hlsl
 	}
 
 	template<>
-	bool ValidateElementWiseUnaryOperator<NNE::Internal::EElementWiseUnaryOperatorType::Selu>(const NNE::FAttributeMap& AttributeMap, TConstArrayView<ENNETensorDataType> InputTypes, TConstArrayView<NNE::FSymbolicTensorShape> InputShapes)
+	bool ValidateElementWiseUnaryOperatorImpl<NNE::Internal::EElementWiseUnaryOperatorType::Selu>(int Version, const NNE::FAttributeMap& AttributeMap, TConstArrayView<ENNETensorDataType> InputTypes, TConstArrayView<NNE::FSymbolicTensorShape> InputShapes)
 	{
 		bool bIsValid = true;
 
@@ -168,7 +218,7 @@ namespace UE::NNERuntimeRDG::Private::Hlsl
 	}
 
 	template<>
-	bool ValidateElementWiseUnaryOperator<NNE::Internal::EElementWiseUnaryOperatorType::Elu>(const NNE::FAttributeMap& AttributeMap, TConstArrayView<ENNETensorDataType> InputTypes, TConstArrayView<NNE::FSymbolicTensorShape> InputShapes)
+	bool ValidateElementWiseUnaryOperatorImpl<NNE::Internal::EElementWiseUnaryOperatorType::Elu>(int Version, const NNE::FAttributeMap& AttributeMap, TConstArrayView<ENNETensorDataType> InputTypes, TConstArrayView<NNE::FSymbolicTensorShape> InputShapes)
 	{
 		bool bIsValid = true;
 
@@ -185,7 +235,7 @@ namespace UE::NNERuntimeRDG::Private::Hlsl
 	}
 
 	template<>
-	bool ValidateElementWiseUnaryOperator<NNE::Internal::EElementWiseUnaryOperatorType::HardSigmoid>(const NNE::FAttributeMap& AttributeMap, TConstArrayView<ENNETensorDataType> InputTypes, TConstArrayView<NNE::FSymbolicTensorShape> InputShapes)
+	bool ValidateElementWiseUnaryOperatorImpl<NNE::Internal::EElementWiseUnaryOperatorType::HardSigmoid>(int Version, const NNE::FAttributeMap& AttributeMap, TConstArrayView<ENNETensorDataType> InputTypes, TConstArrayView<NNE::FSymbolicTensorShape> InputShapes)
 	{
 		bool bIsValid = true;
 
@@ -203,7 +253,7 @@ namespace UE::NNERuntimeRDG::Private::Hlsl
 	}
 
 	template<>
-	bool ValidateElementWiseUnaryOperator<NNE::Internal::EElementWiseUnaryOperatorType::LeakyRelu>(const NNE::FAttributeMap& AttributeMap, TConstArrayView<ENNETensorDataType> InputTypes, TConstArrayView<NNE::FSymbolicTensorShape> InputShapes)
+	bool ValidateElementWiseUnaryOperatorImpl<NNE::Internal::EElementWiseUnaryOperatorType::LeakyRelu>(int Version, const NNE::FAttributeMap& AttributeMap, TConstArrayView<ENNETensorDataType> InputTypes, TConstArrayView<NNE::FSymbolicTensorShape> InputShapes)
 	{
 		bool bIsValid = true;
 
@@ -220,27 +270,41 @@ namespace UE::NNERuntimeRDG::Private::Hlsl
 	}
 
 	template<>
-	bool ValidateElementWiseUnaryOperator<NNE::Internal::EElementWiseUnaryOperatorType::Clip>(const NNE::FAttributeMap& AttributeMap, TConstArrayView<ENNETensorDataType> InputTypes, TConstArrayView<NNE::FSymbolicTensorShape> InputShapes)
+	bool ValidateElementWiseUnaryOperatorImpl<NNE::Internal::EElementWiseUnaryOperatorType::Clip>(int Version, const NNE::FAttributeMap& AttributeMap, TConstArrayView<ENNETensorDataType> InputTypes, TConstArrayView<NNE::FSymbolicTensorShape> InputShapes)
 	{
 		bool bIsValid = true;
 
 		FAttributeValidator AttributeValidator;
-		AttributeValidator.AddOptional(TEXT("min"), ENNEAttributeDataType::Float);
-		AttributeValidator.AddOptional(TEXT("max"), ENNEAttributeDataType::Float);
+		if(Version < 11)
+		{
+			AttributeValidator.AddOptional(TEXT("min"), ENNEAttributeDataType::Float);
+			AttributeValidator.AddOptional(TEXT("max"), ENNEAttributeDataType::Float);
+		}
 		bIsValid &= AttributeValidator.Validate(AttributeMap);
 
 		FInputValidator InputValidator;
 		InputValidator.AddSupportedType(ENNETensorDataType::Float);
 		InputValidator.AddRequired();
+		if(Version >= 11)
+		{
+			InputValidator.AddOptional();
+			InputValidator.AddOptional();
+		}
 		bIsValid &= InputValidator.Validate(InputTypes);
 
 		return bIsValid;
 	}
+
+	template<NNE::Internal::EElementWiseUnaryOperatorType OpType, int Version>
+	bool ValidateElementWiseUnaryOperator(const NNE::FAttributeMap& AttributeMap, TConstArrayView<ENNETensorDataType> InputTypes, TConstArrayView<NNE::FSymbolicTensorShape> InputShapes)
+	{
+		return ValidateElementWiseUnaryOperatorImpl<OpType>(Version, AttributeMap, InputTypes, InputShapes);
+	}
 	
 	bool RegisterElementWiseUnaryOperators(FOperatorRegistryHlsl& Registry)
 	{
 		// Note: support of a particular version is partial with respect to tensor data types (only the most typical ones are usually supported).
-#define OP(Name, Version) Registry.OpAdd({{TEXT(#Name), TEXT("Onnx")}, Version}, CreateElementWiseUnaryOperator<NNE::Internal::EElementWiseUnaryOperatorType::Name>, ValidateElementWiseUnaryOperator<NNE::Internal::EElementWiseUnaryOperatorType::Name>);
+#define OP(Name, Version) Registry.OpAdd({{TEXT(#Name), TEXT("Onnx")}, Version}, CreateElementWiseUnaryOperator<NNE::Internal::EElementWiseUnaryOperatorType::Name, Version>, ValidateElementWiseUnaryOperator<NNE::Internal::EElementWiseUnaryOperatorType::Name, Version>);
 		OP(Abs, 6)
 		OP(Abs, 13)
 		OP(Acos, 7)
@@ -253,6 +317,9 @@ namespace UE::NNERuntimeRDG::Private::Hlsl
 		OP(Ceil, 6)
 		OP(Ceil, 13)
 		OP(Clip, 6)
+		OP(Clip, 11)
+		OP(Clip, 12)
+		OP(Clip, 13)
 		OP(Cos, 7)
 		OP(Cosh, 9)
 		OP(Elu, 6)
