NNE - Add support for Hlsl Softmax-13

Problems:
* RDG Hlsl backend lacks support for Softmax-13.

Solutions:
* Extend RDG Hlsl implementation of softmax to version 13 and add tests.

Tests:
* Win64 Standalone tests: all
* Win64 Editor tests: all

#rnx

#virtualized

diff --git a/Engine/Plugins/Experimental/NNERuntimeRDG/Source/NNEHlslShaders/Internal/NNEHlslShadersSoftmaxCS.h b/Engine/Plugins/Experimental/NNERuntimeRDG/Source/NNEHlslShaders/Internal/NNEHlslShadersSoftmaxCS.h
--- a/Engine/Plugins/Experimental/NNERuntimeRDG/Source/NNEHlslShaders/Internal/NNEHlslShadersSoftmaxCS.h
+++ b/Engine/Plugins/Experimental/NNERuntimeRDG/Source/NNEHlslShaders/Internal/NNEHlslShadersSoftmaxCS.h
@@ -29,7 +29,8 @@ namespace UE::NNEHlslShaders::Internal
 		SHADER_USE_PARAMETER_STRUCT(TSoftmaxCS, FHlslShaderBase)
 
 		class FSoftmaxType : SHADER_PERMUTATION_ENUM_CLASS("SOFTMAX_OPERATOR_TYPE", ESoftmaxOperatorType);
-		using FPermutationDomain = TShaderPermutationDomain<FSoftmaxType>;
+		class FSingleDimension : SHADER_PERMUTATION_BOOL("SINGLE_DIMENSION");
+		using FPermutationDomain = TShaderPermutationDomain<FSoftmaxType,FSingleDimension>;
 
 	public:
 
@@ -37,6 +38,7 @@ namespace UE::NNEHlslShaders::Internal
 			SHADER_PARAMETER(uint32, Num)
 			SHADER_PARAMETER(uint32, ThreadCountX)
 			SHADER_PARAMETER(uint32, AxisSize)
+			SHADER_PARAMETER(uint32, AfterAxisSize)
 			SHADER_PARAMETER_RDG_BUFFER_SRV(Buffer<float>, Input)
 			SHADER_PARAMETER_RDG_BUFFER_SRV(Buffer<float>, InputSumExp)
 			SHADER_PARAMETER_RDG_BUFFER_UAV(RWBuffer<float>, Output)
diff --git a/Engine/Plugins/Experimental/NNERuntimeRDG/Source/NNEHlslShaders/Private/NNEHlslShadersSoftmaxCS.cpp b/Engine/Plugins/Experimental/NNERuntimeRDG/Source/NNEHlslShaders/Private/NNEHlslShadersSoftmaxCS.cpp
--- a/Engine/Plugins/Experimental/NNERuntimeRDG/Source/NNEHlslShaders/Private/NNEHlslShadersSoftmaxCS.cpp
+++ b/Engine/Plugins/Experimental/NNERuntimeRDG/Source/NNEHlslShaders/Private/NNEHlslShadersSoftmaxCS.cpp
@@ -10,6 +10,8 @@ namespace UE::NNEHlslShaders::Internal
 		FGlobalShader::ModifyCompilationEnvironment(InParameters, OutEnvironment);
 
 		OutEnvironment.SetDefine(TEXT("THREADGROUP_SIZE_X"), FSoftmaxConstants::NUM_GROUP_THREADS);
+
+		FPermutationDomain PermutationVector(InParameters.PermutationId);
 	}
 
 	IMPLEMENT_GLOBAL_SHADER(TSoftmaxCS, "/NNE/NNEHlslShadersSoftmax.usf", "Softmax", SF_Compute);
diff --git a/Engine/Plugins/Experimental/NNERuntimeRDG/Source/NNEHlslShaders/Shaders/NNEHlslShadersSoftmax.usf b/Engine/Plugins/Experimental/NNERuntimeRDG/Source/NNEHlslShaders/Shaders/NNEHlslShadersSoftmax.usf
--- a/Engine/Plugins/Experimental/NNERuntimeRDG/Source/NNEHlslShaders/Shaders/NNEHlslShadersSoftmax.usf
+++ b/Engine/Plugins/Experimental/NNERuntimeRDG/Source/NNEHlslShaders/Shaders/NNEHlslShadersSoftmax.usf
@@ -13,6 +13,7 @@ RWBuffer<WORK_TYPE> Output;
 uint Num;
 uint ThreadCountX;
 uint AxisSize;
+uint AfterAxisSize;
 
 #define SOFTMAX_TYPE 0
 #define LOG_SOFTMAX_TYPE 1
@@ -23,7 +24,12 @@ void Softmax(in const uint3 DispatchThreadID : SV_DispatchThreadID)
 	const uint Index = DispatchThreadID.y * ThreadCountX + DispatchThreadID.x;
 	if (Index < Num)
 	{
-		uint SumExpIndex = Index/AxisSize;
+		uint SumExpIndex;
+#if SINGLE_DIMENSION == 0
+		SumExpIndex = Index / AxisSize;
+#else
+		SumExpIndex = (Index / (AfterAxisSize * AxisSize)) * AfterAxisSize + Index % AfterAxisSize;
+#endif
 		WORK_TYPE SumExp = READ(InputSumExp[SumExpIndex]);
 		WORK_TYPE Result = WRITE(exp(READ(Input[Index])) / SumExp);
 
diff --git a/Engine/Plugins/Experimental/NNERuntimeRDG/Source/NNERuntimeRDG/Private/Hlsl/NNERuntimeRDGSoftmax.cpp b/Engine/Plugins/Experimental/NNERuntimeRDG/Source/NNERuntimeRDG/Private/Hlsl/NNERuntimeRDGSoftmax.cpp
--- a/Engine/Plugins/Experimental/NNERuntimeRDG/Source/NNERuntimeRDG/Private/Hlsl/NNERuntimeRDGSoftmax.cpp
+++ b/Engine/Plugins/Experimental/NNERuntimeRDG/Source/NNERuntimeRDG/Private/Hlsl/NNERuntimeRDGSoftmax.cpp
@@ -15,7 +15,7 @@ namespace UE::NNERuntimeRDG::Private::Hlsl
 	/**
 	 * Softmax operator implementation
 	 */
-	template< NNEHlslShaders::Internal::ESoftmaxOperatorType SoftmaxOperatorType >
+	template< NNEHlslShaders::Internal::ESoftmaxOperatorType SoftmaxOperatorType, int Version >
 	class FSoftmax : public FOperatorHlsl
 	{
 
@@ -49,7 +49,7 @@ namespace UE::NNERuntimeRDG::Private::Hlsl
 
 			const int32 InputDimensions = InputTensorDescs[0].GetShape().Rank();
 
-			if (InputDimensions < 2)
+			if (Version <= 11 && InputDimensions < 2)
 			{
 				UE_LOG(LogNNE, Warning, TEXT("Input tensor should be at least 2-D (but got rank %d)"), InputDimensions);
 				return false;
@@ -61,15 +61,19 @@ namespace UE::NNERuntimeRDG::Private::Hlsl
 				return false;
 			}
 
-			const int32 AxisLow = 0;
+			const int32 AxisLow = -InputDimensions;
 			const int32 AxisHigh = InputDimensions - 1;
 
-			Axis = Attributes.GetValueOrDefault<int32>(TEXT("axis"), 1);
+			Axis = Attributes.GetValueOrDefault<int32>(TEXT("axis"), Version <= 11 ? 1 : -1);
 			if (Axis < AxisLow || AxisHigh < Axis)
 			{
 				UE_LOG(LogNNE, Warning, TEXT("Invalid axis (should in the interval [%d, %d], but got %d)"), AxisLow, AxisHigh, Axis);
 				return false;
 			}
+			if(Axis < 0)
+			{
+				Axis += InputDimensions;
+			}
 
 			return true;
 		}
@@ -96,9 +100,18 @@ namespace UE::NNERuntimeRDG::Private::Hlsl
 			// First apply Reduction(exp(x)) to temp buffer
 			TReduceCS::FParameters* ReduceParameters = GraphBuilder.AllocParameters<TReduceCS::FParameters>();
 			TReduceCS::FillInParameters(InputShape.GetData(), Axis, ReduceParameters);
-			ReduceParameters->AxisSize *= ReduceParameters->NumElemAfterAxis;// Softmax flatten the input tensor to a 2D one
-			ReduceParameters->NumElemAfterAxis = 1;
-			const FRDGBufferDesc SumExpBufferDesc = FRDGBufferDesc::CreateBufferDesc(Output.GetElementByteSize(), ReduceParameters->NumElemBeforeAxis);
+			if(Version <= 11)
+			{
+				ReduceParameters->AxisSize *= ReduceParameters->NumElemAfterAxis;// Softmax flatten the input tensor to a 2D one
+				ReduceParameters->NumElemAfterAxis = 1;
+			}
+			const FRDGBufferDesc SumExpBufferDesc = FRDGBufferDesc::CreateBufferDesc(
+				Output.GetElementByteSize(), 
+				Version <= 11 ?
+					ReduceParameters->NumElemBeforeAxis
+					:
+					ReduceParameters->NumElemBeforeAxis * ReduceParameters->NumElemAfterAxis
+				);
 
 			FRDGBufferRef SumExpBuffer = GraphBuilder.CreateBuffer(SumExpBufferDesc, TEXT("NNE.Operator.Hlsl.Softmax.TempBuffer"), ERDGBufferFlags::None);
 
@@ -109,6 +122,10 @@ namespace UE::NNERuntimeRDG::Private::Hlsl
 			const FIntVector ThreadGroupCount = ComputeElementWiseThreadGroups(NumElements, FSoftmaxConstants::NUM_GROUP_THREADS);
 			TSoftmaxCS::FParameters* SoftmaxParameters = GraphBuilder.AllocParameters<TSoftmaxCS::FParameters>();
 			SoftmaxParameters->AxisSize = ReduceParameters->AxisSize;
+			if(Version >= 13)
+			{
+				SoftmaxParameters->AfterAxisSize = ReduceParameters->NumElemAfterAxis;
+			}
 			SoftmaxParameters->Num = NumElements;
 			SoftmaxParameters->ThreadCountX = ThreadGroupCount.X * FSoftmaxConstants::NUM_GROUP_THREADS;
 			SoftmaxParameters->Input = GraphBuilder.CreateSRV(FRDGBufferSRVDesc(Input.GetBuffer(), PF_R32_FLOAT));
@@ -116,6 +133,7 @@ namespace UE::NNERuntimeRDG::Private::Hlsl
 			SoftmaxParameters->Output = GraphBuilder.CreateUAV(FRDGBufferUAVDesc(Output.GetBuffer(), PF_R32_FLOAT));
 			TSoftmaxCS::FPermutationDomain PermutationVector;
 			PermutationVector.Set<TSoftmaxCS::FSoftmaxType>(SoftmaxOperatorType);
+			PermutationVector.Set<TSoftmaxCS::FSingleDimension>(Version >= 13);
 			TShaderMapRef<TSoftmaxCS> ComputeShader(GetGlobalShaderMap(GMaxRHIFeatureLevel), PermutationVector);
 
 			FComputeShaderUtils::AddPass(
@@ -132,7 +150,7 @@ namespace UE::NNERuntimeRDG::Private::Hlsl
 	{
 		bool bIsValid = true;
 
-		// This match version 1 of the Softmax operator, next versions are 11 and 13
+		// This matches versions 1, 11, 13 of the Softmax and LogSoftmax operators
 		// [REDACTED_PROFILE_URL]/onnx/blob/main/docs/Changelog.md#Softmax-1
 		FAttributeValidator AttributeValidator;
 		AttributeValidator.AddOptional(TEXT("axis"), ENNEAttributeDataType::Int32);
@@ -146,20 +164,24 @@ namespace UE::NNERuntimeRDG::Private::Hlsl
 		return bIsValid;
 	}
 
+	template<NNEHlslShaders::Internal::ESoftmaxOperatorType SoftmaxOperatorType, int Version>
 	FOperatorHlsl* CreateSoftmaxOperator()
 	{
-		return new FSoftmax<NNEHlslShaders::Internal::ESoftmaxOperatorType::SOFTMAX>();
-	}
-
-	FOperatorHlsl* CreateLogSoftmaxOperator()
-	{
-		return new FSoftmax<NNEHlslShaders::Internal::ESoftmaxOperatorType::LOG_SOFTMAX>();
+		return new FSoftmax<SoftmaxOperatorType, Version>();
 	}
 
 	bool RegisterSoftmaxOperator(FOperatorRegistryHlsl& Registry)
 	{
-		Registry.OpAdd({ {TEXT("Softmax"), TEXT("Onnx")}, 1 }, CreateSoftmaxOperator, ValidateSoftmaxOperator);
-		Registry.OpAdd({ {TEXT("LogSoftmax"), TEXT("Onnx")}, 1 }, CreateLogSoftmaxOperator, ValidateSoftmaxOperator);
+		// Adds specified version of Softmax and LogSoftmax to the Registry.
+		#define OPS(Version) \
+			Registry.OpAdd({ {TEXT("Softmax"), TEXT("Onnx")}, Version }, CreateSoftmaxOperator<NNEHlslShaders::Internal::ESoftmaxOperatorType::SOFTMAX, Version>, ValidateSoftmaxOperator); \
+			Registry.OpAdd({ {TEXT("LogSoftmax"), TEXT("Onnx")}, Version }, CreateSoftmaxOperator<NNEHlslShaders::Internal::ESoftmaxOperatorType::LOG_SOFTMAX, Version>, ValidateSoftmaxOperator);
+
+		OPS(1)
+		OPS(11)
+		OPS(13)
+
+		#undef OPS
 		
 		return true;
 	}
