NNE - Add RDG Hlsl ScatterND implementation

Problems:
* RDG Hlsl backend lacks an implementation for ScatterND.

Solutions:
* Extend RDG Hlsl implementing ScatterND (all versions).

Tests:
* Win64 Standalone tests: all
* Win64 Editor tests: all

#rnx

#virtualized

diff --git a/Engine/Plugins/Experimental/NNERuntimeRDG/Source/NNEHlslShaders/Internal/NNEHlslShadersScatterNDCS.h b/Engine/Plugins/Experimental/NNERuntimeRDG/Source/NNEHlslShaders/Internal/NNEHlslShadersScatterNDCS.h
new file mode 100644
--- /dev/null
+++ b/Engine/Plugins/Experimental/NNERuntimeRDG/Source/NNEHlslShaders/Internal/NNEHlslShadersScatterNDCS.h
@@ -0,0 +1,57 @@
+// Copyright Epic Games, Inc. All Rights Reserved.
+
+#pragma once
+
+#include "NNEHlslShaderBase.h"
+#include "RenderGraphUtils.h"
+#include "ShaderParameterUtils.h"
+
+namespace UE::NNEHlslShaders::Internal
+{
+
+    enum class EScatterNDReductionType : uint8
+	{
+		None = 0,
+		Add,
+		Mul,
+		Max,
+		Min,
+		MAX
+	};
+
+	class FScatterNDConstants
+	{
+	public:
+
+		static const int32 MAX_NUM_DIMENSIONS{ 8 };
+		static const int32 NUM_GROUP_THREADS{ 768 };
+	};
+
+	class NNEHLSLSHADERS_API FScatterNDCS : public FHlslShaderBase
+	{
+		DECLARE_GLOBAL_SHADER(FScatterNDCS);
+		SHADER_USE_PARAMETER_STRUCT(FScatterNDCS, FHlslShaderBase)
+
+		class FScatterNDNumDimensions : SHADER_PERMUTATION_RANGE_INT("NUM_DIMENSIONS", 1, FScatterNDConstants::MAX_NUM_DIMENSIONS);
+		class FReduceType : SHADER_PERMUTATION_ENUM_CLASS("REDUCE_OPERATOR_TYPE", EScatterNDReductionType);
+		using FPermutationDomain = TShaderPermutationDomain<FScatterNDNumDimensions, FReduceType>;
+
+	public:
+
+		BEGIN_SHADER_PARAMETER_STRUCT(FParameters, )
+			SHADER_PARAMETER_RDG_BUFFER_SRV(Buffer<int>, InputIndices)
+			SHADER_PARAMETER_RDG_BUFFER_SRV(Buffer<float>, InputUpdates)
+			SHADER_PARAMETER_RDG_BUFFER_UAV(RWBuffer<float>, Output)
+			SHADER_PARAMETER_ARRAY(FUintVector4, DataTensorInfo, [FScatterNDConstants::MAX_NUM_DIMENSIONS])
+			SHADER_PARAMETER_ARRAY(FUintVector4, OutputTensorInfo, [FScatterNDConstants::MAX_NUM_DIMENSIONS])
+			SHADER_PARAMETER(uint32, Num)
+			SHADER_PARAMETER(uint32, ThreadCountX)
+			SHADER_PARAMETER(uint32, PartialIndexRank)
+			SHADER_PARAMETER(uint32, SliceSize)
+		END_SHADER_PARAMETER_STRUCT()
+
+		static void ModifyCompilationEnvironment(const FGlobalShaderPermutationParameters& InParameters, FShaderCompilerEnvironment& OutEnvironment);
+	
+        static EScatterNDReductionType ReductionFromString(const TCHAR* StringVal);
+    };
+} // UE::NNEHlslShaders::Internal
diff --git a/Engine/Plugins/Experimental/NNERuntimeRDG/Source/NNEHlslShaders/Private/NNEHlslShadersScatterND.cpp b/Engine/Plugins/Experimental/NNERuntimeRDG/Source/NNEHlslShaders/Private/NNEHlslShadersScatterND.cpp
new file mode 100644
--- /dev/null
+++ b/Engine/Plugins/Experimental/NNERuntimeRDG/Source/NNEHlslShaders/Private/NNEHlslShadersScatterND.cpp
@@ -0,0 +1,29 @@
+// Copyright Epic Games, Inc. All Rights Reserved.
+
+#include "NNEHlslShadersScatterNDCS.h"
+#include "NNE.h"
+
+namespace UE::NNEHlslShaders::Internal
+{
+	void FScatterNDCS::ModifyCompilationEnvironment(const FGlobalShaderPermutationParameters& InParameters, FShaderCompilerEnvironment& OutEnvironment)
+	{
+		FGlobalShader::ModifyCompilationEnvironment(InParameters, OutEnvironment);
+
+		OutEnvironment.SetDefine(TEXT("THREADGROUP_SIZE"), FScatterNDConstants::NUM_GROUP_THREADS);
+
+	}
+
+    EScatterNDReductionType FScatterNDCS::ReductionFromString(const TCHAR* StringVal)
+	{
+		EScatterNDReductionType OutValue = EScatterNDReductionType::None;
+		if (FCString::Stricmp(StringVal, TEXT("none")) == 0) OutValue = EScatterNDReductionType::None;
+		else if (FCString::Stricmp(StringVal, TEXT("add")) == 0) OutValue = EScatterNDReductionType::Add;
+		else if (FCString::Stricmp(StringVal, TEXT("mul")) == 0) OutValue = EScatterNDReductionType::Mul;
+		else if (FCString::Stricmp(StringVal, TEXT("max")) == 0) OutValue = EScatterNDReductionType::Max;
+		else if (FCString::Stricmp(StringVal, TEXT("min")) == 0) OutValue = EScatterNDReductionType::Min;
+
+        return OutValue;
+	}
+
+	IMPLEMENT_GLOBAL_SHADER(FScatterNDCS, "/NNE/NNEHlslShadersScatterND.usf", "ScatterND", SF_Compute);
+} // UE::NNEHlslShaders::Internal
\ No newline at end of file
diff --git a/Engine/Plugins/Experimental/NNERuntimeRDG/Source/NNEHlslShaders/Shaders/NNEHlslShadersScatterND.usf b/Engine/Plugins/Experimental/NNERuntimeRDG/Source/NNEHlslShaders/Shaders/NNEHlslShadersScatterND.usf
new file mode 100644
--- /dev/null
+++ b/Engine/Plugins/Experimental/NNERuntimeRDG/Source/NNEHlslShaders/Shaders/NNEHlslShadersScatterND.usf
@@ -0,0 +1,112 @@
+// Copyright Epic Games, Inc. All Rights Reserved.
+ 
+#include "/Engine/Public/Platform.ush"
+#include "NNEHlslShadersBroadcastHelper.ush"
+
+#define WORK_TYPE float
+#define BUFFER_TYPE float
+#define READ(x) x
+#define WRITE(x) x
+
+Buffer<int> InputIndices;
+Buffer<BUFFER_TYPE> InputUpdates;
+RWBuffer<BUFFER_TYPE> Output;
+uint4 DataTensorInfo[NUM_DIMENSIONS];
+
+uint Num;
+uint ThreadCountX;
+uint PartialIndexRank; // k
+uint SliceSize; // Obtained by multiplying elements in data.shape[k:r-1]
+
+#define STRIDE_IDX 0
+#define SIZE_IDX 1
+
+// Corresponds to EScatterNDReductionType defined in NNEHlslShadersScatterNDCS.h
+#define TYPE_NONE 0
+#define TYPE_ADD 1
+#define TYPE_MUL 2
+#define TYPE_MAX 3
+#define TYPE_MIN 4
+
+uint NormalizeIdx(int Idx, uint DimSize)
+{
+	return Idx < 0 ? (uint) (Idx + DimSize) : (uint) Idx;
+}
+
+// int64_t is only available from SM 6.0, therefore here we need to manually cast to int32.
+int CastInt64ToInt32(uint LSW, int MSW)
+{
+	return (MSW >> 31) == 1 ? - (int) LSW : (int) LSW;
+}
+
+#define STATIC_LOOP(Var, From, To) \
+[unroll] \
+for(uint Var = From; Var < To; ++Var)
+
+void FromIdxToIterator(const uint Idx, const uint4 TensorInfo[NUM_DIMENSIONS], out uint Iterator[NUM_DIMENSIONS])
+{
+	uint Offset = Idx;
+	STATIC_LOOP(DimIdx, 0, NUM_DIMENSIONS)
+	{
+		uint Remainder;
+		DivMod(Offset, TensorInfo[DimIdx][STRIDE_IDX], Iterator[DimIdx], Remainder);
+		Offset = Remainder;
+	}
+}
+
+uint FromIteratorToIdx(const uint4 TensorInfo[NUM_DIMENSIONS], const uint Iterator[NUM_DIMENSIONS])
+{
+	uint Offset = 0;
+	STATIC_LOOP(DimIdx, 0, NUM_DIMENSIONS)
+	{
+		Offset += Iterator[DimIdx] * TensorInfo[DimIdx][STRIDE_IDX];
+	}
+	return Offset;
+}
+
+void ApplyReduction(const uint OutputIdx, const uint UpdatesIdx)
+{
+	#if REDUCE_OPERATOR_TYPE == TYPE_NONE
+		Output[OutputIdx] = WRITE(READ(InputUpdates[UpdatesIdx]));
+	#elif REDUCE_OPERATOR_TYPE == TYPE_ADD
+		Output[OutputIdx] = WRITE(READ(Output[OutputIdx]) + READ(InputUpdates[UpdatesIdx]));
+	#elif REDUCE_OPERATOR_TYPE == TYPE_MUL
+		Output[OutputIdx] = WRITE(READ(Output[OutputIdx]) * READ(InputUpdates[UpdatesIdx]));
+	#elif REDUCE_OPERATOR_TYPE == TYPE_MAX
+		Output[OutputIdx] = WRITE(max(READ(Output[OutputIdx]), READ(InputUpdates[UpdatesIdx])));
+	#elif REDUCE_OPERATOR_TYPE == TYPE_MIN
+		Output[OutputIdx] = WRITE(min(READ(Output[OutputIdx]), READ(InputUpdates[UpdatesIdx])));
+	#endif
+}
+
+
+[numthreads(THREADGROUP_SIZE, 1, 1)]
+void ScatterND(in const uint3 DispatchThreadID : SV_DispatchThreadID)
+{
+	// See [REDACTED_PROFILE_URL]/onnx/blob/main/docs/Operators.md#ScatterND
+	// Parallelizing over updates, of shape: indices.shape[0:q-2] ++ data.shape[k:r-1]
+	const uint Idx = DispatchThreadID.y * ThreadCountX + DispatchThreadID.x;
+
+	if(Idx < Num)
+	{
+		uint ReplacementIdx = Idx / SliceSize;
+		uint ElemIdx = Idx % SliceSize;
+
+		uint OutputIterator[NUM_DIMENSIONS];
+
+		// Set the index of the element within the slice
+		FromIdxToIterator(ElemIdx, DataTensorInfo, OutputIterator);
+
+		// Then retrieve and set the index of the slice
+		for (uint DimIdx = 0; DimIdx < PartialIndexRank; ++DimIdx)
+		{
+			uint IndicesAccessIdx = ReplacementIdx * PartialIndexRank + DimIdx;
+			int RawIdx = CastInt64ToInt32( asuint(READ(InputIndices[IndicesAccessIdx * 2])), READ(InputIndices[IndicesAccessIdx * 2 + 1]) );
+			OutputIterator[DimIdx] = NormalizeIdx(RawIdx, DataTensorInfo[DimIdx][SIZE_IDX]);
+		}
+
+		uint OutputIdx = FromIteratorToIdx(DataTensorInfo, OutputIterator);
+		
+		ApplyReduction(OutputIdx, Idx);
+	}
+}
diff --git a/Engine/Plugins/Experimental/NNERuntimeRDG/Source/NNERuntimeRDG/Private/Hlsl/NNERuntimeRDGScatterND.cpp b/Engine/Plugins/Experimental/NNERuntimeRDG/Source/NNERuntimeRDG/Private/Hlsl/NNERuntimeRDGScatterND.cpp
new file mode 100644
--- /dev/null
+++ b/Engine/Plugins/Experimental/NNERuntimeRDG/Source/NNERuntimeRDG/Private/Hlsl/NNERuntimeRDGScatterND.cpp
@@ -0,0 +1,194 @@
+// Copyright Epic Games, Inc. All Rights Reserved.
+
+#include "NNERuntimeRDGScatterND.h"
+#include "NNEHlslShadersScatterNDCS.h"
+#include "NNERuntimeRDGHlslHelper.h"
+#include "NNETensor.h"
+#include "NNETypes.h"
+#include "RenderGraphUtils.h"
+#include "Algo/Compare.h"
+
+namespace UE::NNERuntimeRDG::Private::Hlsl
+{
+	DECLARE_GPU_STAT_NAMED(FNNEOperatorScatterND, TEXT("NNE.Operator.Hlsl.ScatterND"));
+
+	/**
+	 * ScatterND operator implementation
+	 */
+	class FScatterND : public FOperatorHlsl
+	{
+	public:
+
+		FScatterND() {}
+		virtual ~FScatterND() = default;
+
+	private:
+		UE::NNEHlslShaders::Internal::EScatterNDReductionType ReductionType;
+
+	public:
+
+		virtual int PrepareOutputs(TConstArrayView<NNE::Internal::FTensorRef> InputTensors, TArrayView<NNE::Internal::FTensorRef> OutputTensors) override
+		{
+			check(InputTensors.Num() == 3);
+			check(OutputTensors.Num() == 1);
+
+			const NNE::Internal::FTensor& Input = *InputTensors[0];
+			TConstArrayView<uint32> InputShape = Input.GetShape().GetData();
+			const NNE::Internal::FTensor& Indices = *InputTensors[1];
+			TConstArrayView<uint32> IndicesShape = Indices.GetShape().GetData();
+			const NNE::Internal::FTensor& Updates = *InputTensors[2];
+			TConstArrayView<uint32> UpdatesShape = Updates.GetShape().GetData();
+
+			if (IndicesShape.Last() > (uint32) InputShape.Num())
+			{
+				UE_LOG(LogNNE, Warning, TEXT("Last dimension in the shape of `indices` (%d) must be less than input rank (%d)."), IndicesShape.Last(), InputShape.Num());
+				return -1;
+			}
+			if (UpdatesShape.Num() != (IndicesShape.Num() - 1) + (InputShape.Num() - IndicesShape.Last()))
+			{
+				UE_LOG(LogNNE, Warning, TEXT("Rank of `updates` should equal (q - 1) + (r - k), with q rank of `indices`, r rank of `data` and `k` last dimension of `indices`' shape."));
+				return -1;
+			}
+			if (!Algo::Compare( UpdatesShape.Slice(0, IndicesShape.Num() - 1), IndicesShape.Slice(0, IndicesShape.Num() - 1) ))
+			{
+				UE_LOG(LogNNE, Warning, TEXT("updates.shape[0:q-1] should match indices.shape[0:q-1]."));
+				return -1;
+			}
+			if (!Algo::Compare( UpdatesShape.RightChop(IndicesShape.Num() - 1), InputShape.RightChop(IndicesShape.Last()) ))
+			{
+				UE_LOG(LogNNE, Warning, TEXT("updates.shape[q-1:] should match data.shape[k:]."));
+				return -1;
+			}
+			//NOTE: we need to limit the volume of Input to the int32 maximum expressible value due to unavailability of int64 type in shaders.
+			if (Input.GetShape().Volume() > (uint64) MAX_int32)
+			{
+				UE_LOG(LogNNE, Warning, TEXT("This implementation of ScatterND only supports input tensors up to a volume of %u."), MAX_int32);
+				return -1;
+			}
+
+			OutputTensors[0]->SetShape(NNE::FTensorShape::Make(InputShape));
+
+			return 0;
+		};
+
+		virtual bool Initialize(TConstArrayView<NNE::FTensorDesc> InputTensorDescs, TConstArrayView<NNE::FTensorDesc> OutputTensorDescs, const NNE::FAttributeMap& Attributes) override
+		{
+            using namespace UE::NNEHlslShaders::Internal;
+			check(InputTensorDescs.Num() == 3);
+			check(OutputTensorDescs.Num() == 1);
+
+			const NNE::FTensorDesc& InputDesc = InputTensorDescs[0];
+			const int32 InputRank = InputDesc.GetShape().Rank();
+			
+			ReductionType = FScatterNDCS::ReductionFromString(*Attributes.GetValueOrDefault<FString>(TEXT("reduction"), TEXT("none")));
+
+			return true;
+		}
+
+		virtual void Dispatch(FRDGBuilder& GraphBuilder, TConstArrayView<FTensorRDGRef> InputTensors, TConstArrayView<FTensorRDGRef> OutputTensors) override
+		{
+			using namespace UE::NNEHlslShaders::Internal;
+
+			check(InputTensors.Num() == 3);
+			check(OutputTensors.Num() == 1);
+			check(InputTensors[0] != nullptr);
+			check(InputTensors[1] != nullptr);
+			check(InputTensors[2] != nullptr);
+			check(OutputTensors[0] != nullptr);
+
+			const FTensorRDG& Input = *InputTensors[0];
+			const FTensorRDG& Indices = *InputTensors[1];
+			const FTensorRDG& Updates = *InputTensors[2];
+			const FTensorRDG& Output = *OutputTensors[0];
+
+			// NOTE: Indices tensor is int64, but UE lacks support of int64 buffers. Here we use a 32-bit pixel format and in the shader we reinterpret two words as int64.
+			const FRDGBufferSRVRef InputIndicesSRV = GraphBuilder.CreateSRV(FRDGBufferSRVDesc(Indices.GetBuffer(), PF_R32_SINT));
+			const FRDGBufferSRVRef InputUpdatesSRV = GraphBuilder.CreateSRV(FRDGBufferSRVDesc(Updates.GetBuffer(), PF_R32_FLOAT));
+			const FRDGBufferUAVRef OutputUAV = GraphBuilder.CreateUAV(FRDGBufferUAVDesc(Output.GetBuffer(), PF_R32_FLOAT));
+
+			const FIntVector ThreadGroupCount = ComputeElementWiseThreadGroups(Updates.GetVolume(), FScatterNDConstants::NUM_GROUP_THREADS);
+
+			// Set parameters
+			FScatterNDCS::FParameters* Params = GraphBuilder.AllocParameters<FScatterNDCS::FParameters>();
+			
+			Params->InputIndices = InputIndicesSRV;
+			Params->InputUpdates = InputUpdatesSRV;
+			Params->Output = OutputUAV;
+			Params->Num = Updates.GetVolume();
+			Params->ThreadCountX = ThreadGroupCount.X * FScatterNDConstants::NUM_GROUP_THREADS;
+
+			FillTensorStrideShaderParameters(Input, Params->DataTensorInfo, /*Idx*/0);
+			FillTensorSizeShaderParameters(Input, Params->DataTensorInfo, /*Idx*/1);
+			static_assert(NNE::FTensorShape::MaxRank <= NXRT_TENSORSTRIDEINFO_MAX_NUM_DIMENSIONS);
+			uint32 PartialIndexRank = Indices.GetShape().GetData().Last();
+			check(PartialIndexRank <= (uint32) Input.GetShape().Rank());
+			Params->PartialIndexRank = PartialIndexRank;
+
+
+			auto ComputeSliceSize = [] (TConstArrayView<uint32> DataShape, uint32 SliceStartIdx) -> uint32 
+				{
+					uint32 Res = 1;
+					for(uint32 DimIdx = SliceStartIdx; DimIdx < (uint32) DataShape.Num(); ++DimIdx)
+					{
+						Res *= DataShape[DimIdx];
+					}
+					return Res;
+				};
+			Params->SliceSize = ComputeSliceSize(Input.GetShape().GetData(), PartialIndexRank);
+
+			FScatterNDCS::FPermutationDomain PermutationVector;
+			PermutationVector.Set<FScatterNDCS::FScatterNDNumDimensions>(Output.GetShape().Rank());
+			PermutationVector.Set<FScatterNDCS::FReduceType>(ReductionType);
+
+			TShaderMapRef<FScatterNDCS> ComputeShader(GetGlobalShaderMap(GMaxRHIFeatureLevel), PermutationVector);
+
+			RDG_EVENT_SCOPE(GraphBuilder, "NNE.Operator.Hlsl.ScatterND");
+			RDG_GPU_STAT_SCOPE(GraphBuilder, FNNEOperatorScatterND);
+			
+			AddCopyBufferPass(GraphBuilder, Output.GetBuffer(), Input.GetBuffer());
+
+			FComputeShaderUtils::AddPass(
+				GraphBuilder,
+				RDG_EVENT_NAME("NNE.Operator.Hlsl.ScatterND.Dispatch"),
+				ERDGPassFlags::Compute | ERDGPassFlags::NeverCull,
+				ComputeShader,
+				Params,
+				ThreadGroupCount);
+		}
+	};
+
+	bool ValidateScatterNDOperator(const NNE::FAttributeMap& AttributeMap, TConstArrayView<ENNETensorDataType> InputTypes, TConstArrayView<NNE::FSymbolicTensorShape> InputScatterNDs)
+	{
+		bool bIsValid = true;
+
+		FAttributeValidator AttributeValidator;
+		AttributeValidator.AddOptional(TEXT("reduction"), ENNEAttributeDataType::String);
+		bIsValid &= AttributeValidator.Validate(AttributeMap);
+
+		FInputValidator InputValidator;
+		InputValidator.SetTemplateCount(2);
+		InputValidator.AddSupportedType(ENNETensorDataType::Float);
+		InputValidator.AddSupportedType(ENNETensorDataType::Int64, /*TemplateIdx*/ 1);
+		InputValidator.AddRequired();
+		InputValidator.AddRequired(/*TemplateIdx*/ 1);
+		InputValidator.AddRequired();
+		bIsValid &= InputValidator.Validate(InputTypes);
+
+		return bIsValid;
+	}
+
+	FOperatorHlsl* CreateScatterNDOperator()
+	{
+		return new FScatterND();
+	}
+
+	bool RegisterScatterNDOperator(FOperatorRegistryHlsl& Registry)
+	{
+		// Note: support of a particular version is partial with respect to tensor data types (only the most typical ones are usually supported).
+		Registry.OpAdd({{TEXT("ScatterND"), TEXT("Onnx")}, 11}, CreateScatterNDOperator, ValidateScatterNDOperator);
+		Registry.OpAdd({{TEXT("ScatterND"), TEXT("Onnx")}, 13}, CreateScatterNDOperator, ValidateScatterNDOperator);
+		Registry.OpAdd({{TEXT("ScatterND"), TEXT("Onnx")}, 16}, CreateScatterNDOperator, ValidateScatterNDOperator);
+		Registry.OpAdd({{TEXT("ScatterND"), TEXT("Onnx")}, 18}, CreateScatterNDOperator, ValidateScatterNDOperator);
+		return true;
+	}
+} // UE::NNERuntimeRDG::Private::Hlsl
diff --git a/Engine/Plugins/Experimental/NNERuntimeRDG/Source/NNERuntimeRDG/Private/Hlsl/NNERuntimeRDGScatterND.h b/Engine/Plugins/Experimental/NNERuntimeRDG/Source/NNERuntimeRDG/Private/Hlsl/NNERuntimeRDGScatterND.h
new file mode 100644
--- /dev/null
+++ b/Engine/Plugins/Experimental/NNERuntimeRDG/Source/NNERuntimeRDG/Private/Hlsl/NNERuntimeRDGScatterND.h
@@ -0,0 +1,11 @@
+// Copyright Epic Games, Inc. All Rights Reserved.
+
+#pragma once
+
+#include "NNERuntimeRDGBase.h"
+#include "NNERuntimeRDGHlslOp.h"
+
+namespace UE::NNERuntimeRDG::Private::Hlsl
+{
+	bool RegisterScatterNDOperator(FOperatorRegistryHlsl& Registry);
+} // UE::NNERuntimeRDG::Private::Hlsl
diff --git a/Engine/Plugins/Experimental/NNERuntimeRDG/Source/NNERuntimeRDG/Private/NNERuntimeRDGHlsl.cpp b/Engine/Plugins/Experimental/NNERuntimeRDG/Source/NNERuntimeRDG/Private/NNERuntimeRDGHlsl.cpp
--- a/Engine/Plugins/Experimental/NNERuntimeRDG/Source/NNERuntimeRDG/Private/NNERuntimeRDGHlsl.cpp
+++ b/Engine/Plugins/Experimental/NNERuntimeRDG/Source/NNERuntimeRDG/Private/NNERuntimeRDGHlsl.cpp
@@ -36,6 +36,7 @@
 #include "Hlsl/NNERuntimeRDGPool.h"
 #include "Hlsl/NNERuntimeRDGReduce.h"
 #include "Hlsl/NNERuntimeRDGReshape.h"
+#include "Hlsl/NNERuntimeRDGScatterND.h"
 #include "Hlsl/NNERuntimeRDGShape.h"
 #include "Hlsl/NNERuntimeRDGSize.h"
 #include "Hlsl/NNERuntimeRDGSlice.h"
@@ -77,6 +78,7 @@ bool UNNERuntimeRDGHlslImpl::Init()
 	RegisterPoolOperators(*Registry);
 	RegisterReduceOperators(*Registry);
 	RegisterReshapeOperator(*Registry);
+	RegisterScatterNDOperator(*Registry);
 	RegisterShapeOperator(*Registry);
 	RegisterSizeOperator(*Registry);
 	RegisterSliceOperator(*Registry);
