NNE - Hlsl Gather add support for int64 Indices tensor

Problems:
* RDG Hlsl Gather implementation lacks support for Indices tensor when it's int64.

Solutions:
* Extend RDG Hlsl Gather implementation with support for int64 Indices.

Tests:
* Win64 Standalone tests: all
* Win64 Editor tests: all

#rnx

#virtualized

diff --git a/Engine/Plugins/Experimental/NNERuntimeRDG/Source/NNEHlslShaders/Internal/NNEHlslShadersGatherCS.h b/Engine/Plugins/Experimental/NNERuntimeRDG/Source/NNEHlslShaders/Internal/NNEHlslShadersGatherCS.h
--- a/Engine/Plugins/Experimental/NNERuntimeRDG/Source/NNEHlslShaders/Internal/NNEHlslShadersGatherCS.h
+++ b/Engine/Plugins/Experimental/NNERuntimeRDG/Source/NNEHlslShaders/Internal/NNEHlslShadersGatherCS.h
@@ -25,7 +25,8 @@ namespace UE::NNEHlslShaders::Internal
 		SHADER_USE_PARAMETER_STRUCT(TGatherCS, FHlslShaderBase)
 
 		class FGatherNumOutputDimensions : SHADER_PERMUTATION_RANGE_INT("NUM_OUTPUT_DIMENSIONS", 1, FGatherConstants::MAX_NUM_DIMENSIONS);
-		using FPermutationDomain = TShaderPermutationDomain<FGatherNumOutputDimensions>;
+		class FGather64BitIndices : SHADER_PERMUTATION_BOOL("HAS_64BIT_INDICES");
+		using FPermutationDomain = TShaderPermutationDomain<FGatherNumOutputDimensions,FGather64BitIndices>;
 
 	public:
 
@@ -34,7 +35,7 @@ namespace UE::NNEHlslShaders::Internal
 			SHADER_PARAMETER(int32, OutputSize)
 			SHADER_PARAMETER(int32, NumDataDimensions)
 			SHADER_PARAMETER(int32, NumIndicesDimensions)
-			SHADER_PARAMETER_ARRAY(FIntVector4, DataStride_IndicesStride_OutputStride, [FGatherConstants::MAX_NUM_DIMENSIONS])
+			SHADER_PARAMETER_ARRAY(FIntVector4, DataStride_IndicesStride_OutputStride_DataSizes, [FGatherConstants::MAX_NUM_DIMENSIONS])
 			SHADER_PARAMETER_ARRAY(FVector4f, OneDivDataStride_OneDivIndicesStride_OneDivOutputStride, [FGatherConstants::MAX_NUM_DIMENSIONS])
 			SHADER_PARAMETER_RDG_BUFFER_SRV(Buffer<float>, Data)
 			SHADER_PARAMETER_RDG_BUFFER_SRV(Buffer<int32>, Indices)
diff --git a/Engine/Plugins/Experimental/NNERuntimeRDG/Source/NNEHlslShaders/Private/NNEHlslShadersGatherCS.cpp b/Engine/Plugins/Experimental/NNERuntimeRDG/Source/NNEHlslShaders/Private/NNEHlslShadersGatherCS.cpp
--- a/Engine/Plugins/Experimental/NNERuntimeRDG/Source/NNEHlslShaders/Private/NNEHlslShadersGatherCS.cpp
+++ b/Engine/Plugins/Experimental/NNERuntimeRDG/Source/NNEHlslShaders/Private/NNEHlslShadersGatherCS.cpp
@@ -46,7 +46,7 @@ namespace UE::NNEHlslShaders::Internal
 		int32 DataStride = 1;
 		for (int32 i = Data.GetShape().Rank() - 1; i >= 0; i--)
 		{
-			Parameters.DataStride_IndicesStride_OutputStride[i].X = DataStride;
+			Parameters.DataStride_IndicesStride_OutputStride_DataSizes[i].X = DataStride;
 			Parameters.OneDivDataStride_OneDivIndicesStride_OneDivOutputStride[i].X = 1.0 / (float)DataStride;
 			DataStride *= Data.GetShape().GetData()[i];
 		}
@@ -54,7 +54,7 @@ namespace UE::NNEHlslShaders::Internal
 		int32 IndicesStride = 1;
 		for (int32 i = Indices.GetShape().Rank() - 1; i >= 0; i--)
 		{
-			Parameters.DataStride_IndicesStride_OutputStride[i].Y = IndicesStride;
+			Parameters.DataStride_IndicesStride_OutputStride_DataSizes[i].Y = IndicesStride;
 			Parameters.OneDivDataStride_OneDivIndicesStride_OneDivOutputStride[i].Y = 1.0 / (float)IndicesStride;
 			IndicesStride *= Indices.GetShape().GetData()[i];
 		}
@@ -62,10 +62,15 @@ namespace UE::NNEHlslShaders::Internal
 		int32 OutputStride = 1;
 		for (int32 i = OutputShape.Num() - 1; i >= 0; i--)
 		{
-			Parameters.DataStride_IndicesStride_OutputStride[i].Z = OutputStride;
+			Parameters.DataStride_IndicesStride_OutputStride_DataSizes[i].Z = OutputStride;
 			Parameters.OneDivDataStride_OneDivIndicesStride_OneDivOutputStride[i].Z = 1.0 / (float)OutputStride;
 			OutputStride *= OutputShape[i];
 		}
+
+		for (int32 Idx = 0; Idx < Data.GetShape().Rank(); ++Idx)
+		{
+			Parameters.DataStride_IndicesStride_OutputStride_DataSizes[Idx].W = Data.GetShape().GetData()[Idx];
+		}
 	}
 
 	// template <typename DataElementType, typename IndicesElementType>
diff --git a/Engine/Plugins/Experimental/NNERuntimeRDG/Source/NNEHlslShaders/Shaders/NNEHlslShadersGather.usf b/Engine/Plugins/Experimental/NNERuntimeRDG/Source/NNEHlslShaders/Shaders/NNEHlslShadersGather.usf
--- a/Engine/Plugins/Experimental/NNERuntimeRDG/Source/NNEHlslShaders/Shaders/NNEHlslShadersGather.usf
+++ b/Engine/Plugins/Experimental/NNERuntimeRDG/Source/NNEHlslShaders/Shaders/NNEHlslShadersGather.usf
@@ -22,13 +22,25 @@ int Axis;
 int OutputSize;
 int NumDataDimensions;
 int NumIndicesDimensions;
-int4 DataStride_IndicesStride_OutputStride[MAX_NUM_DIMENSIONS];
+int4 DataStride_IndicesStride_OutputStride_DataSizes[MAX_NUM_DIMENSIONS];
 float4 OneDivDataStride_OneDivIndicesStride_OneDivOutputStride[MAX_NUM_DIMENSIONS];
 
+
 Buffer<DATA_BUFFER_TYPE> Data;
 Buffer<INDICES_BUFFER_TYPE> Indices;
 RWBuffer<DATA_BUFFER_TYPE> Output;
 
+uint NormalizeIdx(int Idx, uint DimSize)
+{
+	return Idx < 0 ? (uint) (Idx + DimSize) : (uint) Idx;
+}
+
+// int64_t is only available from SM 6.0, therefore here we need to manually cast to int32.
+int CastInt64ToInt32(uint LSW, int MSW)
+{
+	return (MSW >> 31) == 1 ? - (int) LSW : (int) LSW;
+}
+
 [numthreads(NUM_GROUP_THREADS, 1, 1)]
 void Gather(in const uint3 DispatchThreadID : SV_DispatchThreadID)
 {
@@ -44,7 +56,7 @@ void Gather(in const uint3 DispatchThreadID : SV_DispatchThreadID)
 	for (int i = 0; i < NUM_OUTPUT_DIMENSIONS; i++)
 	{
 		OutputIndices[i] = (INDICES_WORK_TYPE)(OneDivDataStride_OneDivIndicesStride_OneDivOutputStride[i].z * (float)DispatchThreadIndex);
-		DispatchThreadIndex -= OutputIndices[i] * (INDICES_WORK_TYPE)DataStride_IndicesStride_OutputStride[i].z;
+		DispatchThreadIndex -= OutputIndices[i] * (INDICES_WORK_TYPE)DataStride_IndicesStride_OutputStride_DataSizes[i].z;
 	}
 
 	// Split the output indices into data and indices indices 
@@ -67,17 +79,27 @@ void Gather(in const uint3 DispatchThreadID : SV_DispatchThreadID)
 	INDICES_WORK_TYPE IndicesIndex = 0;
 	for (int i = 0; i < NumIndicesDimensions; i++)
 	{
-		IndicesIndex += IndicesIndices[i] * DataStride_IndicesStride_OutputStride[i].y;
+		IndicesIndex += IndicesIndices[i] * DataStride_IndicesStride_OutputStride_DataSizes[i].y;
 	}
 
 	// Look up and assign the missing data index
-	DataIndices[Axis] = INDICES_BUFFER_TO_WORK_TYPE(Indices[IndicesIndex]);
+	#if HAS_64BIT_INDICES == 0
+		DataIndices[Axis] = NormalizeIdx(
+			INDICES_BUFFER_TO_WORK_TYPE(Indices[IndicesIndex]), 
+			DataStride_IndicesStride_OutputStride_DataSizes[Axis].w
+		);
+	#else
+		DataIndices[Axis] = NormalizeIdx(
+			INDICES_BUFFER_TO_WORK_TYPE(CastInt64ToInt32(Indices[IndicesIndex * 2], Indices[IndicesIndex * 2 + 1])),
+			DataStride_IndicesStride_OutputStride_DataSizes[Axis].w
+		);
+	#endif
 
 	// Compute the flattened data index
 	INDICES_WORK_TYPE DataIndex = 0;
 	for (int i = 0; i < NumDataDimensions; i++)
 	{
-		DataIndex += DataIndices[i] * DataStride_IndicesStride_OutputStride[i].x;
+		DataIndex += DataIndices[i] * DataStride_IndicesStride_OutputStride_DataSizes[i].x;
 	}
 
 	// Write the result
diff --git a/Engine/Plugins/Experimental/NNERuntimeRDG/Source/NNERuntimeRDG/Private/Hlsl/NNERuntimeRDGGather.cpp b/Engine/Plugins/Experimental/NNERuntimeRDG/Source/NNERuntimeRDG/Private/Hlsl/NNERuntimeRDGGather.cpp
--- a/Engine/Plugins/Experimental/NNERuntimeRDG/Source/NNERuntimeRDG/Private/Hlsl/NNERuntimeRDGGather.cpp
+++ b/Engine/Plugins/Experimental/NNERuntimeRDG/Source/NNERuntimeRDG/Private/Hlsl/NNERuntimeRDGGather.cpp
@@ -129,11 +129,13 @@ namespace UE::NNERuntimeRDG::Private::Hlsl
 			TGatherCS::FParameters* Parameters = GraphBuilder.AllocParameters<TGatherCS::FParameters>();
 			TGatherCS::FillInParameters(Axis, Data, Indices, *Parameters);
 			Parameters->Data = GraphBuilder.CreateSRV(FRDGBufferSRVDesc(Data.GetBuffer(), PF_R32_FLOAT));
+			// NOTE: Indices tensor can be int64, but UE lacks support of int64 buffers. Here we use a 32-bit pixel format and in the shader we reinterpret two words as int64.
 			Parameters->Indices = GraphBuilder.CreateSRV(FRDGBufferSRVDesc(Indices.GetBuffer(), PF_R32_FLOAT));
 			Parameters->Output = GraphBuilder.CreateUAV(FRDGBufferUAVDesc(Output.GetBuffer(), PF_R32_FLOAT));
 
 			TGatherCS::FPermutationDomain PermutationVector;
 			PermutationVector.Set<TGatherCS::FGatherNumOutputDimensions>(Output.GetShape().Rank());
+			PermutationVector.Set<TGatherCS::FGather64BitIndices>(Indices.GetDataType() == ENNETensorDataType::Int64);
 			TShaderMapRef<TGatherCS> ComputeShader(GetGlobalShaderMap(GMaxRHIFeatureLevel), PermutationVector);
 
 			FIntVector ThreadGroupCount = TGatherCS::GetGroupCount(*Parameters);
