NNE - Add Hlsl Resize (partial support)

Problems:
* RDG Hlsl Resize implementation is missing.

Solutions:
* Add partial support for Hlsl Resize: nearest and linear modes.
* Missing: cubic mode, antialias, Sizes tensor (requires RDG Hlsl infrastructure support for empty tensors), non-float data types for input/roi/scales tensor, axes attribute, exclude_outside attribute.

Tests:
* Win64 Standalone tests: all
* Win64 Editor tests: all

#rnx

#virtualized

diff --git a/Engine/Plugins/Experimental/NNERuntimeRDG/Source/NNEHlslShaders/Internal/NNEHlslShadersResizeCS.h b/Engine/Plugins/Experimental/NNERuntimeRDG/Source/NNEHlslShaders/Internal/NNEHlslShadersResizeCS.h
new file mode 100644
--- /dev/null
+++ b/Engine/Plugins/Experimental/NNERuntimeRDG/Source/NNEHlslShaders/Internal/NNEHlslShadersResizeCS.h
@@ -0,0 +1,77 @@
+// Copyright Epic Games, Inc. All Rights Reserved.
+
+#pragma once
+
+#include "NNEHlslShaderBase.h"
+#include "RenderGraphUtils.h"
+#include "ShaderParameterUtils.h"
+
+namespace UE::NNEHlslShaders::Internal
+{
+
+	enum class ECoordTransMode : uint8
+	{
+		HalfPixel = 0,
+		HalfPixelSymmetric,
+		PytorchHalfPixel,
+		AlignCorners,
+		Asymmetric,
+		TfHalfPixelForNn,
+		TfCropAndResize,
+		MAX
+	};
+
+	enum class EMode : uint8
+	{
+		Nearest = 0,
+		Linear,
+		Cubic,
+		MAX
+	};
+
+	enum class ENearestMode : uint8
+	{
+		RoundPreferFloor = 0,
+		RoundPreferCeil,
+		Floor,
+		Ceil,
+		MAX
+	};
+
+	class FResizeConstants
+	{
+	public:
+
+		static const int32 MAX_NUM_DIMENSIONS{ 8 };
+		static const int32 NUM_GROUP_THREADS{ 256 };
+	};
+
+	class NNEHLSLSHADERS_API FResizeCS : public FHlslShaderBase
+	{
+		DECLARE_GLOBAL_SHADER(FResizeCS);
+		SHADER_USE_PARAMETER_STRUCT(FResizeCS, FHlslShaderBase)
+
+		class FResizeNumDimensions : SHADER_PERMUTATION_RANGE_INT("NUM_DIMENSIONS", 1, FResizeConstants::MAX_NUM_DIMENSIONS);
+		class FMode : SHADER_PERMUTATION_ENUM_CLASS("MODE", EMode);
+		class FNearestMode : SHADER_PERMUTATION_ENUM_CLASS("NEAREST_MODE", ENearestMode);
+		class FCoordTransMode : SHADER_PERMUTATION_ENUM_CLASS("COORD_TRANS_MODE", ECoordTransMode);
+		using FPermutationDomain = TShaderPermutationDomain<FResizeNumDimensions, FMode, FNearestMode, FCoordTransMode>;
+
+	public:
+		BEGIN_SHADER_PARAMETER_STRUCT(FParameters, )
+			SHADER_PARAMETER_RDG_BUFFER_SRV(Buffer<float>, Input)
+			SHADER_PARAMETER_ARRAY(FUintVector4, InputTensorInfo, [FResizeConstants::MAX_NUM_DIMENSIONS])
+			SHADER_PARAMETER_RDG_BUFFER_UAV(RWBuffer<float>, Output)
+			SHADER_PARAMETER_ARRAY(FUintVector4, OutputTensorInfo, [FResizeConstants::MAX_NUM_DIMENSIONS])
+			SHADER_PARAMETER_ARRAY(FVector4f, ScalesData, [FResizeConstants::MAX_NUM_DIMENSIONS])
+			SHADER_PARAMETER(uint32, Num)
+			SHADER_PARAMETER(uint32, ThreadCountX)
+		END_SHADER_PARAMETER_STRUCT()
+
+		static bool ShouldCompilePermutation(const FGlobalShaderPermutationParameters& InParameters);
+		static void ModifyCompilationEnvironment(const FGlobalShaderPermutationParameters& InParameters, FShaderCompilerEnvironment& OutEnvironment);
+		static ECoordTransMode CoordTransModeFromString(const TCHAR* StringVal);
+		static EMode ModeFromString(const TCHAR* StringVal);
+		static ENearestMode NearestModeFromString(const TCHAR* StringVal);
+	};
+} // UE::NNEHlslShaders::Internal
diff --git a/Engine/Plugins/Experimental/NNERuntimeRDG/Source/NNEHlslShaders/Private/NNEHlslShadersResizeCS.cpp b/Engine/Plugins/Experimental/NNERuntimeRDG/Source/NNEHlslShaders/Private/NNEHlslShadersResizeCS.cpp
new file mode 100644
--- /dev/null
+++ b/Engine/Plugins/Experimental/NNERuntimeRDG/Source/NNEHlslShaders/Private/NNEHlslShadersResizeCS.cpp
@@ -0,0 +1,73 @@
+// Copyright Epic Games, Inc. All Rights Reserved.
+
+#include "NNEHlslShadersResizeCS.h"
+#include "NNE.h"
+
+namespace UE::NNEHlslShaders::Internal
+{
+	bool FResizeCS::ShouldCompilePermutation(const FGlobalShaderPermutationParameters& Parameters)
+	{
+		if (!FHlslShaderBase::ShouldCompilePermutation(Parameters))
+		{
+			return false;
+		}
+
+		FPermutationDomain PermutationVector(Parameters.PermutationId);
+
+		EMode Mode = PermutationVector.Get<FResizeCS::FMode>();
+		
+		//NOTE: cubic interpolation not currently supported
+		if(Mode == EMode::Cubic)
+		{
+			return false;
+		}
+
+		return true;
+	}
+
+	void FResizeCS::ModifyCompilationEnvironment(const FGlobalShaderPermutationParameters& InParameters, FShaderCompilerEnvironment& OutEnvironment)
+	{
+		FGlobalShader::ModifyCompilationEnvironment(InParameters, OutEnvironment);
+
+		OutEnvironment.SetDefine(TEXT("THREADGROUP_SIZE_X"), FResizeConstants::NUM_GROUP_THREADS);
+
+		FPermutationDomain PermutationVector(InParameters.PermutationId);
+	}
+
+	ECoordTransMode FResizeCS::CoordTransModeFromString(const TCHAR* StringVal)
+	{
+		ECoordTransMode OutValue = ECoordTransMode::HalfPixel;
+		if (FCString::Stricmp(StringVal, TEXT("half_pixel")) == 0) OutValue = ECoordTransMode::HalfPixel;
+		else if (FCString::Stricmp(StringVal, TEXT("half_pixel_symmetric")) == 0) OutValue = ECoordTransMode::HalfPixelSymmetric;
+		else if (FCString::Stricmp(StringVal, TEXT("pytorch_half_pixel")) == 0) OutValue = ECoordTransMode::PytorchHalfPixel;
+		else if (FCString::Stricmp(StringVal, TEXT("align_corners")) == 0) OutValue = ECoordTransMode::AlignCorners;
+		else if (FCString::Stricmp(StringVal, TEXT("asymmetric")) == 0) OutValue = ECoordTransMode::Asymmetric;
+		else if (FCString::Stricmp(StringVal, TEXT("tf_half_pixel_for_nn")) == 0) OutValue = ECoordTransMode::TfHalfPixelForNn;
+		else if (FCString::Stricmp(StringVal, TEXT("tf_crop_and_resize")) == 0) OutValue = ECoordTransMode::TfCropAndResize;
+
+        return OutValue;
+	}
+
+	EMode FResizeCS::ModeFromString(const TCHAR* StringVal)
+	{
+		EMode OutValue = EMode::Nearest;
+		if (FCString::Stricmp(StringVal, TEXT("nearest")) == 0) OutValue = EMode::Nearest;
+		else if (FCString::Stricmp(StringVal, TEXT("linear")) == 0) OutValue = EMode::Linear;
+		else if (FCString::Stricmp(StringVal, TEXT("cubic")) == 0) OutValue = EMode::Cubic;
+
+        return OutValue;
+	}
+
+	ENearestMode FResizeCS::NearestModeFromString(const TCHAR* StringVal)
+	{
+		ENearestMode OutValue = ENearestMode::RoundPreferFloor;
+		if (FCString::Stricmp(StringVal, TEXT("round_prefer_floor")) == 0) OutValue = ENearestMode::RoundPreferFloor;
+		else if (FCString::Stricmp(StringVal, TEXT("round_prefer_ceil")) == 0) OutValue = ENearestMode::RoundPreferCeil;
+		else if (FCString::Stricmp(StringVal, TEXT("floor")) == 0) OutValue = ENearestMode::Floor;
+		else if (FCString::Stricmp(StringVal, TEXT("ceil")) == 0) OutValue = ENearestMode::Ceil;
+
+        return OutValue;
+	}
+
+	IMPLEMENT_GLOBAL_SHADER(FResizeCS, "/NNE/NNEHlslShadersResize.usf", "Resize", SF_Compute);
+} // UE::NNEHlslShaders::Internal
\ No newline at end of file
diff --git a/Engine/Plugins/Experimental/NNERuntimeRDG/Source/NNEHlslShaders/Shaders/NNEHlslShadersResize.usf b/Engine/Plugins/Experimental/NNERuntimeRDG/Source/NNEHlslShaders/Shaders/NNEHlslShadersResize.usf
new file mode 100644
--- /dev/null
+++ b/Engine/Plugins/Experimental/NNERuntimeRDG/Source/NNEHlslShaders/Shaders/NNEHlslShadersResize.usf
@@ -0,0 +1,213 @@
+// Copyright Epic Games, Inc. All Rights Reserved.
+
+#include "/Engine/Public/Platform.ush"
+#include "NNEHlslShadersBroadcastHelper.ush"
+#include "/Engine/Private/DoubleFloat.ush"
+
+static const uint NumPointsLinear = pow(2,NUM_DIMENSIONS);
+
+Buffer<float> Input;
+uint4 InputTensorInfo[NUM_DIMENSIONS];
+RWBuffer<float> Output;
+uint4 OutputTensorInfo[NUM_DIMENSIONS];
+float4 ScalesData[NUM_DIMENSIONS];
+uint Num;
+uint ThreadCountX;
+
+#define STRIDE_IDX 0
+#define SIZE_IDX 1
+#define ADJUSTMENT_IDX 2
+#define SCALES_IDX 0
+#define ROI_START_IDX 2
+#define ROI_END_IDX 3
+
+// Corresponds to ECoordTransMode
+#define COORD_TRANS_HALF_PIXEL 0
+#define COORD_TRANS_HALF_PIXEL_SYMMETRIC 1
+#define COORD_TRANS_PYTORCH_HALF_PIXEL 2
+#define COORD_TRANS_ALIGN_CORNERS 3
+#define COORD_TRANS_ASYMMETRIC 4
+#define COORD_TRANS_TF_HALF_PIXEL_FOR_NN 5
+#define COORD_TRANS_TF_CROP_AND_RESIZE 6
+
+// Corresponds to ECoordTransMode enum
+#define NEAREST_MODE_ROUND_PREFER_FLOOR 0
+#define NEAREST_MODE_ROUND_PREFER_CEIL 1
+#define NEAREST_MODE_FLOOR 2
+#define NEAREST_MODE_CEIL 3
+
+// Corresponds to EMode enum
+#define MODE_NEAREST 0
+#define MODE_LINEAR 1
+#define MODE_CUBIC 2
+
+#define STATIC_LOOP(Var, From, To) \
+[unroll] \
+for(uint Var = From; Var < To; ++Var)
+
+void FromIdxToIterator(const uint Idx, const uint4 TensorInfo[NUM_DIMENSIONS], out uint Iterator[NUM_DIMENSIONS])
+{
+	uint Offset = Idx;
+	STATIC_LOOP(DimIdx, 0, NUM_DIMENSIONS)
+	{
+		uint Remainder;
+		DivMod(Offset, TensorInfo[DimIdx][STRIDE_IDX], Iterator[DimIdx], Remainder);
+		Offset = Remainder;
+	}
+}
+
+uint FromIteratorToIdx(const uint4 TensorInfo[NUM_DIMENSIONS], const uint Iterator[NUM_DIMENSIONS])
+{
+	uint Offset = 0;
+	STATIC_LOOP(DimIdx, 0, NUM_DIMENSIONS)
+	{
+		Offset += Iterator[DimIdx] * TensorInfo[DimIdx][STRIDE_IDX];
+	}
+	return Offset;
+}
+
+float TransformCoord(float Resized, uint DimIdx)
+{
+	float Scale = ScalesData[DimIdx][SCALES_IDX];
+
+	// Computations of formulas which combine sums and divisions are done in double precision since single precision is affected by some heavy optimizations
+	// that lead to very inaccurate results. This way we obtain the same input pixels as the ORT implementations.
+
+	#if COORD_TRANS_MODE == COORD_TRANS_HALF_PIXEL
+		return DFSubtractDemote(DFDivide(DFTwoSum(Resized, 0.5f), Scale), 0.5f);
+	#elif COORD_TRANS_MODE == COORD_TRANS_HALF_PIXEL_SYMMETRIC
+		float Adjustment = asfloat(OutputTensorInfo[DimIdx][ADJUSTMENT_IDX]); // Decode adjustment to float
+		FDFScalar Center = DFDivide(InputTensorInfo[DimIdx][SIZE_IDX], 2.0f);
+		return DFSubtractDemote(DFAdd(DFMultiply(Center, DFSubtract(1.0f, Adjustment)), DFDivide(DFTwoSum(Resized, 0.5f), Scale)), 0.5f);
+	#elif COORD_TRANS_MODE == COORD_TRANS_PYTORCH_HALF_PIXEL
+		return OutputTensorInfo[DimIdx][SIZE_IDX] > 1 ? DFSubtractDemote(DFDivide(DFTwoSum(Resized, 0.5f), Scale), 0.5f) : 0.0f;
+	#elif COORD_TRANS_MODE == COORD_TRANS_ALIGN_CORNERS
+		return DFMultiplyDemote(Resized, DFDivide((InputTensorInfo[DimIdx][SIZE_IDX] - 1), (OutputTensorInfo[DimIdx][SIZE_IDX] - 1)));
+	#elif COORD_TRANS_MODE == COORD_TRANS_ASYMMETRIC
+		return DFDemote(DFDivide(Resized, Scale));
+	#elif COORD_TRANS_MODE == COORD_TRANS_TF_HALF_PIXEL_FOR_NN
+		return DFDemote(DFDivide(DFTwoSum(Resized, 0.5f), Scale));
+	#else // COORD_TRANS_MODE == COORD_TRANS_TF_CROP_AND_RESIZE
+		float Start = asfloat(InputTensorInfo[DimIdx][ROI_START_IDX]); // Decode ROI start to float
+		float End = asfloat(InputTensorInfo[DimIdx][ROI_END_IDX]); // Decode ROI end to float
+		return OutputTensorInfo[DimIdx][SIZE_IDX] > 1 ? 
+			DFAddDemote(DFMultiply(Start, (InputTensorInfo[DimIdx][SIZE_IDX] - 1)), DFDivide(DFMultiply(DFMultiply(Resized, DFSubtract(End, Start)), (InputTensorInfo[DimIdx][SIZE_IDX] - 1)), (OutputTensorInfo[DimIdx][SIZE_IDX] - 1))) 
+			: 
+			DFMultiplyDemote(DFMultiply(0.5f, DFTwoSum(Start, End)), (InputTensorInfo[DimIdx][SIZE_IDX] - 1));
+	#endif
+
+}
+
+void GetNearestInput(uint OutputIndex, out uint Iterator[NUM_DIMENSIONS])
+{
+	uint OutputIterator[NUM_DIMENSIONS];
+	
+	FromIdxToIterator(OutputIndex, OutputTensorInfo, OutputIterator);
+
+	STATIC_LOOP(DimIdx, 0, NUM_DIMENSIONS)
+	{
+		float InputCoord = TransformCoord((float) OutputIterator[DimIdx], DimIdx);
+		
+		int Dim = (int)
+			#if NEAREST_MODE == NEAREST_MODE_FLOOR
+				floor(InputCoord)
+			#elif NEAREST_MODE == NEAREST_MODE_CEIL
+				ceil(InputCoord)
+			#elif NEAREST_MODE == NEAREST_MODE_ROUND_PREFER_FLOOR
+				( InputCoord == (int) InputCoord + 0.5f ? floor(InputCoord) : round(InputCoord) )
+			#else // NEAREST_MODE == NEAREST_MODE_ROUND_PREFER_CEIL
+				( InputCoord == (int) InputCoord + 0.5f ? ceil(InputCoord) : round(InputCoord) )
+			#endif
+			;
+		
+		Dim = clamp(Dim, 0, (int) InputTensorInfo[DimIdx][SIZE_IDX] - 1);
+		
+		Iterator[DimIdx] = (uint) Dim;
+	}
+
+}
+
+#if MODE == MODE_NEAREST
+
+[numthreads(THREADGROUP_SIZE_X, 1, 1)]
+void Resize(in const uint3 DispatchThreadID : SV_DispatchThreadID)
+{
+	const uint OutputIndex = DispatchThreadID.y * ThreadCountX + DispatchThreadID.x;
+	
+	if (OutputIndex < Num)
+	{
+		uint InputIterator[NUM_DIMENSIONS];
+
+		GetNearestInput(OutputIndex, InputIterator);
+
+		uint InputIndex = FromIteratorToIdx(InputTensorInfo, InputIterator);		
+
+		Output[OutputIndex] = Input[InputIndex];
+
+	}
+}
+
+#elif MODE == MODE_LINEAR
+
+[numthreads(THREADGROUP_SIZE_X, 1, 1)]
+void Resize(in const uint3 DispatchThreadID : SV_DispatchThreadID)
+{
+	const uint OutputIndex = DispatchThreadID.y * ThreadCountX + DispatchThreadID.x;
+	
+	if (OutputIndex < Num)
+	{
+		float Weights[NUM_DIMENSIONS];
+		uint InputIterator[NUM_DIMENSIONS];
+
+		{
+			uint OutputIterator[NUM_DIMENSIONS];
+		
+			FromIdxToIterator(OutputIndex, OutputTensorInfo, OutputIterator);
+
+			STATIC_LOOP(DimIdx, 0, NUM_DIMENSIONS)
+			{
+				float InputCoord = InputTensorInfo[DimIdx][SIZE_IDX] == 1 ? 
+						(float) OutputIterator[DimIdx]
+						:
+						TransformCoord((float) OutputIterator[DimIdx], DimIdx);
+				InputCoord = max(0.0f, min(InputCoord, (float)(InputTensorInfo[DimIdx][SIZE_IDX] - 1)));
+				InputIterator[DimIdx] = (uint) floor(InputCoord);
+
+				Weights[DimIdx] = (InputIterator[DimIdx] >= InputTensorInfo[DimIdx][SIZE_IDX] - 1) ? 0.5f : frac(InputCoord);
+			}
+
+		}
+
+		float OutputValue = 0.0f;
+
+		STATIC_LOOP(PointIdx, 0, NumPointsLinear)
+		{
+			uint PointCoords[NUM_DIMENSIONS];
+			float Factor = 1.0f;
+
+			{
+
+				STATIC_LOOP(DimIdx, 0, NUM_DIMENSIONS)
+				{
+					uint PointCoord = (PointIdx & (1 << DimIdx)) >> DimIdx;
+					PointCoords[DimIdx] = InputIterator[DimIdx] + PointCoord;
+					if(PointCoords[DimIdx] >= (InputTensorInfo[DimIdx][SIZE_IDX] - 1))
+					{
+						PointCoords[DimIdx] = InputTensorInfo[DimIdx][SIZE_IDX] - 1;
+					}
+					Factor *= PointCoord == 1 ? Weights[DimIdx] : (1.0f - Weights[DimIdx]);
+
+				}
+
+			}
+
+			float PointValue = Input[FromIteratorToIdx(InputTensorInfo, PointCoords)];
+			OutputValue += PointValue * Factor;
+		}
+
+		Output[OutputIndex] = OutputValue;
+	}
+}
+
+
+#endif
\ No newline at end of file
diff --git a/Engine/Plugins/Experimental/NNERuntimeRDG/Source/NNERuntimeRDG/Private/Hlsl/NNERuntimeRDGResize.cpp b/Engine/Plugins/Experimental/NNERuntimeRDG/Source/NNERuntimeRDG/Private/Hlsl/NNERuntimeRDGResize.cpp
new file mode 100644
--- /dev/null
+++ b/Engine/Plugins/Experimental/NNERuntimeRDG/Source/NNERuntimeRDG/Private/Hlsl/NNERuntimeRDGResize.cpp
@@ -0,0 +1,474 @@
+// Copyright Epic Games, Inc. All Rights Reserved.
+
+#include "NNERuntimeRDGResize.h"
+#include "NNEHlslShadersResizeCS.h"
+
+#include "NNERuntimeRDGHlslHelper.h"
+#include "NNETensor.h"
+#include "NNETypes.h"
+#include "RenderGraphUtils.h"
+#include "Algo/MinElement.h"
+#include "Algo/MaxElement.h"
+
+namespace UE::NNERuntimeRDG::Private::Hlsl
+{
+	DECLARE_GPU_STAT_NAMED(FNNEOperatorResize, TEXT("NNE.Operator.Hlsl.Resize"));
+
+
+	/**
+	 * Resize operator implementation
+	 */
+	class FResize : public FOperatorHlsl
+	{
+	public:
+
+		FResize() {}
+		virtual ~FResize() = default;
+
+	private:
+
+		enum EKeepAspectRatioPolicy
+		{
+			Stretch,
+			NotLarger,
+			NotSmaller
+		};
+
+		static EKeepAspectRatioPolicy KeepAspectRatioPolicyFromString(const TCHAR* StringVal)
+		{
+			EKeepAspectRatioPolicy OutValue = EKeepAspectRatioPolicy::Stretch;
+			if (FCString::Stricmp(StringVal, TEXT("stretch")) == 0) OutValue = EKeepAspectRatioPolicy::Stretch;
+			else if (FCString::Stricmp(StringVal, TEXT("not_larger")) == 0) OutValue = EKeepAspectRatioPolicy::NotLarger;
+			else if (FCString::Stricmp(StringVal, TEXT("not_smaller")) == 0) OutValue = EKeepAspectRatioPolicy::NotSmaller;
+
+			return OutValue;
+		}
+		
+		TArray<int32> Axes;
+		UE::NNEHlslShaders::Internal::ECoordTransMode CoordTransMode;
+		float CubicCoeffA;
+		int32 ExcludeOutside;
+		float ExtrapolationValue;
+		EKeepAspectRatioPolicy KeepAspectRatioPolicy;
+		UE::NNEHlslShaders::Internal::EMode Mode;
+		UE::NNEHlslShaders::Internal::ENearestMode NearestMode;
+		TArray<float> ScalesData;
+		TArray<float> Adjustments; // Necessary for "half_pixel_symmetric" coordinate transform mode
+		TArray<float> RegionOfInterest;
+		
+
+	public:
+
+		virtual int PrepareOutputs(TConstArrayView<NNE::Internal::FTensorRef> InputTensors, TArrayView<NNE::Internal::FTensorRef> OutputTensors) override
+		{
+			check(InputTensors.Num() >= 1 && InputTensors.Num() <= 4);
+			check(OutputTensors.Num() == 1);
+
+			
+			if(InputTensors.Num() < 3)
+			{
+				UE_LOG(LogNNE, Warning, TEXT("Hlsl Resize: one of tensors `scales` or `sizes` MUST be provided."));
+				return -1;
+			}
+
+			if(InputTensors.Num() == 4 && InputTensors[2]->GetVolume() != 0)
+			{
+				UE_LOG(LogNNE, Warning, TEXT("Hlsl Resize: both `sizes` and `scales` were provided."));
+				return -1;
+			}
+
+			for(int Idx = 1; Idx < InputTensors.Num(); ++Idx)
+			{
+				if(!InputTensors[Idx]->HasPreparedData())
+				{
+					UE_LOG(LogNNE, Warning, TEXT("Hlsl Resize: tensor %s should be constant for shape inference to succeed."), *InputTensors[Idx]->GetName());
+					return -1;
+				}
+			}
+
+			const NNE::Internal::FTensor& Input = *InputTensors[0];
+			const NNE::Internal::FTensor& Roi = *InputTensors[1];
+
+			if(Input.GetShape().Rank() < 1)
+			{
+				UE_LOG(LogNNE, Warning, TEXT("Hlsl Resize: input tensor must have rank >= 1 (name %s)."), *Input.GetName());
+				return -1;
+			}
+
+			if(Input.GetDataType() != ENNETensorDataType::Float)
+			{
+				UE_LOG(LogNNE, Warning, TEXT("Hlsl Resize: only float data type supported for input tensor (name %s)."), *Input.GetName());
+				return -1;
+			}
+
+			if(Roi.GetShape().Rank() != 1)
+			{
+				UE_LOG(LogNNE, Warning, TEXT("Hlsl Resize: `roi` tensor must have rank 1 (name %s)."), *Roi.GetName());
+				return -1;
+			}
+
+			if(Roi.GetDataType() != ENNETensorDataType::Float)
+			{
+				UE_LOG(LogNNE, Warning, TEXT("Hlsl Resize: only float data type supported for`roi` tensor (name %s)."), *Roi.GetName());
+				return -1;
+			}
+
+			if(!Roi.HasPreparedData())
+			{
+				UE_LOG(LogNNE, Warning, TEXT("Hlsl Resize: `roi` tensor could not be made constant. (name %s)."), *Roi.GetName());
+				return -1;
+			}
+			
+			RegionOfInterest = Roi.GetPreparedData<float>();
+
+			if(CoordTransMode == UE::NNEHlslShaders::Internal::ECoordTransMode::TfCropAndResize
+				&& RegionOfInterest.Num() != 2 * Input.GetShape().Rank())
+			{
+				UE_LOG(LogNNE, Warning, TEXT("Hlsl Resize: `roi` tensor (name %s) must have 2 * N length."), *Roi.GetName());
+				return -1;
+			}
+
+			TArray<uint32> OutputShapeData;
+			OutputShapeData.SetNumUninitialized(Input.GetShape().Rank());
+
+			ScalesData.SetNumUninitialized(Input.GetShape().Rank());
+		
+			if(InputTensors.Num() == 3)
+			{
+				const NNE::Internal::FTensor& Scales = *InputTensors[2];
+				if(Scales.GetDataType() != ENNETensorDataType::Float)
+				{
+					UE_LOG(LogNNE, Warning, TEXT("Hlsl Resize: `scales` tensor (name %s) must have float data type."), *Scales.GetName());
+					return -1;
+				}
+				if(!Scales.HasPreparedData())
+				{
+					UE_LOG(LogNNE, Warning, TEXT("Hlsl Resize: `scales` tensor could not be made constant. (name %s)."), *Scales.GetName());
+					return -1;
+				}
+
+				
+				
+				if(CoordTransMode == UE::NNEHlslShaders::Internal::ECoordTransMode::HalfPixelSymmetric)
+				{
+					Adjustments.SetNumUninitialized(Input.GetShape().Rank());
+				}
+
+				for (int32 Idx = 0; Idx < Input.GetShape().Rank(); ++Idx)
+				{
+					OutputShapeData[Idx] = (uint32) FMath::FloorToInt32(
+						//NOTE: documentation erroneously says that for coordinate_transform_mode = tf_crop_and_resize the output shape should be computed as:
+						// (float) Input.GetShape().GetData()[Idx] * (Roi.GetShape().GetData()[Idx + Input.GetShape().Rank()] - Roi.GetShape().GetData()[Idx]) * Scales.GetPreparedData<float>()[Idx]
+						// however it should be computed like for any other coordinate_transform_mode:
+							(float) Input.GetShape().GetData()[Idx] * Scales.GetPreparedData<float>()[Idx]
+						);
+
+					ScalesData[Idx] = Scales.GetPreparedData<float>()[Idx];
+					
+					if(CoordTransMode == UE::NNEHlslShaders::Internal::ECoordTransMode::HalfPixelSymmetric)
+					{
+						Adjustments[Idx] = (float) OutputShapeData[Idx] / ((float) Input.GetShape().GetData()[Idx] * Scales.GetPreparedData<float>()[Idx]);
+					}
+				}
+
+			}
+			
+			if(InputTensors.Num() == 4)
+			{
+				const NNE::Internal::FTensor& Sizes = *InputTensors[3];
+				if(Sizes.GetDataType() != ENNETensorDataType::Int64)
+				{
+					UE_LOG(LogNNE, Warning, TEXT("Hlsl Resize: `sizes` tensor (name %s) must have int64 data type."), *Sizes.GetName());
+					return -1;
+				}
+				if(!Sizes.HasPreparedData())
+				{
+					UE_LOG(LogNNE, Warning, TEXT("Hlsl Resize: `sizes` tensor could not be made constant. (name %s)."), *Sizes.GetName());
+					return -1;
+				}
+
+				TArray<float> Ratios;
+				{
+					Ratios.SetNumUninitialized(Input.GetShape().Rank());
+					for (int32 Idx = 0; Idx < Input.GetShape().Rank(); ++Idx)
+					{
+						Ratios[Idx] = (float) Sizes.GetPreparedData<int64>()[Idx] / (float) Input.GetShape().GetData()[Idx];
+					}
+				}
+				float MinRatio = *Algo::MinElement(Ratios);
+				float MaxRatio = *Algo::MaxElement(Ratios);
+
+				for (int32 Idx = 0; Idx < Input.GetShape().Rank(); ++Idx)
+				{
+					switch(KeepAspectRatioPolicy)
+					{
+					case EKeepAspectRatioPolicy::Stretch:
+						OutputShapeData[Idx] = Sizes.GetPreparedData<int64>()[Idx];
+						break;
+					case EKeepAspectRatioPolicy::NotLarger:
+						OutputShapeData[Idx] = (uint32) FMath::RoundToInt32(MinRatio * Input.GetShape().GetData()[Idx]);
+						break;
+					case EKeepAspectRatioPolicy::NotSmaller:
+						OutputShapeData[Idx] = (uint32) FMath::RoundToInt32(MaxRatio * Input.GetShape().GetData()[Idx]);
+						break;
+					};
+
+					ScalesData[Idx] = (float) OutputShapeData[Idx] / (float) Input.GetShape().GetData()[Idx];
+				}
+				
+			}
+
+			NNE::FTensorShape OutputShape = NNE::FTensorShape::Make(OutputShapeData);
+			OutputTensors[0]->SetShape(OutputShape);
+
+			
+			return 0;
+		};
+
+		virtual bool Initialize(TConstArrayView<NNE::FTensorDesc> InputTensorDescs, TConstArrayView<NNE::FTensorDesc> OutputTensorDescs, const NNE::FAttributeMap& Attributes) override
+		{
+			check(InputTensorDescs.Num() >= 1 && InputTensorDescs.Num() <= 4);
+			check(OutputTensorDescs.Num() == 1);
+
+			// NOTE: implementation for antialiasing is missing
+			if(Attributes.GetValueOrDefault<int32>(TEXT("antialias"), 0) == 1)
+			{
+				UE_LOG(LogNNE, Warning, TEXT("Hlsl Resize: antialias not yet supported."));
+				return false;
+			}
+
+			Axes = Attributes.GetValueOrDefault<TArray<int32>>(TEXT("axes"), TArray<int32>{});
+			
+			for (int32& Axis : Axes)
+			{
+				if (Axis > InputTensorDescs[0].GetShape().Rank() || Axis < -InputTensorDescs[0].GetShape().Rank())
+				{
+					UE_LOG(LogNNE, Warning, TEXT("Hlsl Resize: 'Axes' attribute should contain value be in the range [-r,r] with r being the rank of the input (name: %s) however got %d while rank is %d."), *InputTensorDescs[0].GetName(), Axis, InputTensorDescs[0].GetShape().Rank());
+					return false;
+				}
+
+				if (Axis < 0)
+				{
+					Axis += InputTensorDescs[0].GetShape().Rank();
+				}
+			}
+
+			// NOTE: `axes` attribute not yet supported
+			if(Axes.Num() != 0)
+			{
+				UE_LOG(LogNNE, Warning, TEXT("Hlsl Resize: `axes` attribute not yet supported."));
+				return false;
+			}
+
+			CoordTransMode = UE::NNEHlslShaders::Internal::FResizeCS::CoordTransModeFromString(*Attributes.GetValueOrDefault<FString>(TEXT("coordinate_transformation_mode"), TEXT("half_pixel")));
+			CubicCoeffA = Attributes.GetValueOrDefault<float>(TEXT("cubic_coeff_a"), -0.75f);
+			ExcludeOutside = Attributes.GetValueOrDefault<int32>(TEXT("exclude_outside"), 0);
+			// NOTE: `exclude_outside` attribute not yet supported
+			if(ExcludeOutside == 1)
+			{
+				UE_LOG(LogNNE, Warning, TEXT("Hlsl Resize: `exclude_outside` attribute not yet supported."));
+				return false;
+			}
+
+			ExtrapolationValue = Attributes.GetValueOrDefault<float>(TEXT("extrapolation_value"), 0.0f);
+			KeepAspectRatioPolicy = KeepAspectRatioPolicyFromString(*Attributes.GetValueOrDefault<FString>(TEXT("keep_aspect_ratio_policy"), TEXT("stretch")));
+			Mode = UE::NNEHlslShaders::Internal::FResizeCS::ModeFromString(*Attributes.GetValueOrDefault<FString>(TEXT("mode"), TEXT("nearest")));
+			//NOTE: cubic interpolation not yet supported
+			if(Mode == UE::NNEHlslShaders::Internal::EMode::Cubic)
+			{
+				UE_LOG(LogNNE, Warning, TEXT("Hlsl Resize: cubic interpolation not yet supported."));
+				return false;
+			}
+
+			NearestMode = UE::NNEHlslShaders::Internal::FResizeCS::NearestModeFromString(*Attributes.GetValueOrDefault<FString>(TEXT("nearest_mode"), TEXT("round_prefer_floor")));
+			
+			return true;
+		}
+
+		virtual void Dispatch(FRDGBuilder& GraphBuilder, TConstArrayView<FTensorRDGRef> InputTensors, TConstArrayView<FTensorRDGRef> OutputTensors) override
+		{
+			using namespace UE::NNEHlslShaders::Internal;
+
+			check(InputTensors.Num() >= 1 && InputTensors.Num() <= 4);
+			check(OutputTensors.Num() == 1);
+			for(int Idx = 0; Idx < InputTensors.Num(); ++Idx)
+			{
+				check(InputTensors[Idx] != nullptr);
+			}
+			check(OutputTensors[0] != nullptr);
+
+			const FTensorRDG& Input = *InputTensors[0];
+			const FTensorRDG& Output = *OutputTensors[0];
+
+			RDG_EVENT_SCOPE(GraphBuilder, "NNE.Operator.Hlsl.Resize");
+			RDG_GPU_STAT_SCOPE(GraphBuilder, FNNEOperatorResize);
+
+			const FRDGBufferSRVRef InputSRV = GraphBuilder.CreateSRV(FRDGBufferSRVDesc(Input.GetBuffer(), PF_R32_FLOAT));
+			const FRDGBufferUAVRef OutputUAV = GraphBuilder.CreateUAV(FRDGBufferUAVDesc(Output.GetBuffer(), PF_R32_FLOAT));
+
+			const FIntVector ThreadGroupCount = ComputeElementWiseThreadGroups(Output.GetVolume(), FResizeConstants::NUM_GROUP_THREADS);
+
+			// Set parameters
+			FResizeCS::FParameters* Params = GraphBuilder.AllocParameters<FResizeCS::FParameters>();
+			
+			Params->Input = InputSRV;
+			Params->Output = OutputUAV;
+			Params->Num = Output.GetVolume();
+			Params->ThreadCountX = ThreadGroupCount.X * FResizeConstants::NUM_GROUP_THREADS;
+
+			FillTensorStrideShaderParameters(Input, Params->InputTensorInfo, /*Idx*/0);
+			FillTensorSizeShaderParameters(Input, Params->InputTensorInfo, /*Idx*/1);
+			FillTensorStrideShaderParameters(Output, Params->OutputTensorInfo, /*Idx*/0);
+			FillTensorSizeShaderParameters(Output, Params->OutputTensorInfo, /*Idx*/1);
+
+			for(int DimIdx = 0; DimIdx < Input.GetShape().Rank(); ++DimIdx)
+			{
+				// NOTE: floats are encoded as unit32s and then decoded in the shader code
+
+				if(CoordTransMode == UE::NNEHlslShaders::Internal::ECoordTransMode::HalfPixelSymmetric)
+				{
+					// Set adjustments
+					Params->OutputTensorInfo[DimIdx][/*Idx*/ 2] = BitCast<uint32>(Adjustments[DimIdx]);
+				}
+				Params->ScalesData[DimIdx][/*Idx*/ 0] = ScalesData[DimIdx];
+
+				if(CoordTransMode == UE::NNEHlslShaders::Internal::ECoordTransMode::TfCropAndResize)
+				{
+					// Set ROI start indices
+					Params->InputTensorInfo[DimIdx][/*Idx*/ 2] = BitCast<uint32>(RegionOfInterest[DimIdx]);
+					// Set ROI end indices
+					Params->InputTensorInfo[DimIdx][/*Idx*/ 3] = BitCast<uint32>(RegionOfInterest[Input.GetShape().Rank() + DimIdx]);
+				}
+			}
+
+			FResizeCS::FPermutationDomain PermutationVector;
+			PermutationVector.Set<FResizeCS::FResizeNumDimensions>(Output.GetShape().Rank());
+			PermutationVector.Set<FResizeCS::FMode>(Mode);
+			PermutationVector.Set<FResizeCS::FNearestMode>(NearestMode);
+			PermutationVector.Set<FResizeCS::FCoordTransMode>(CoordTransMode);
+
+			TShaderMapRef<FResizeCS> ComputeShader(GetGlobalShaderMap(GMaxRHIFeatureLevel), PermutationVector);
+
+			RDG_EVENT_SCOPE(GraphBuilder, "NNE.Operator.Hlsl.Resize");
+			RDG_GPU_STAT_SCOPE(GraphBuilder, FNNEOperatorResize);
+			
+			FComputeShaderUtils::AddPass(
+				GraphBuilder,
+				RDG_EVENT_NAME("NNE.Operator.Hlsl.Resize.Dispatch"),
+				ERDGPassFlags::Compute | ERDGPassFlags::NeverCull,
+				ComputeShader,
+				Params,
+				ThreadGroupCount);
+		}
+	};
+
+	bool ValidateResizeOperator(const NNE::FAttributeMap& AttributeMap, TConstArrayView<ENNETensorDataType> InputTypes, TConstArrayView<NNE::FSymbolicTensorShape> InputShapes)
+	{
+		bool bIsValid = true;
+
+		FAttributeValidator AttributeValidator;
+		AttributeValidator.AddOptional(TEXT("antialias"), ENNEAttributeDataType::Int32);
+		AttributeValidator.AddOptional(TEXT("axes"), ENNEAttributeDataType::Int32Array);
+		AttributeValidator.AddOptional(TEXT("coordinate_transformation_mode"), ENNEAttributeDataType::String);
+		AttributeValidator.AddOptional(TEXT("cubic_coeff_a"), ENNEAttributeDataType::Float);
+		AttributeValidator.AddOptional(TEXT("exclude_outside"), ENNEAttributeDataType::Int32);
+		AttributeValidator.AddOptional(TEXT("extrapolation_value"), ENNEAttributeDataType::Float);
+		AttributeValidator.AddOptional(TEXT("keep_aspect_ratio_policy"), ENNEAttributeDataType::String);
+		AttributeValidator.AddOptional(TEXT("mode"), ENNEAttributeDataType::String);
+		AttributeValidator.AddOptional(TEXT("nearest_mode"), ENNEAttributeDataType::String);
+		bIsValid &= AttributeValidator.Validate(AttributeMap);
+
+		if(InputTypes.Num() < 3 || InputTypes.Num() > 4)
+		{
+			UE_LOG(LogNNE, Warning, TEXT("Got a total of '%d' inputs but should be between 3 and 4."), InputTypes.Num());
+			return false;
+		}
+		if(InputTypes[0] != ENNETensorDataType::Float)
+		{
+			UE_LOG(LogNNE, Warning, TEXT("Input tensor must be of type float."));
+			return false;
+		}
+		if(InputTypes[1] != ENNETensorDataType::Float)
+		{
+			UE_LOG(LogNNE, Warning, TEXT("Roi tensor must be of type float."));
+			return false;
+		}
+		if(InputShapes[1].Rank() != 1)
+		{
+			UE_LOG(LogNNE, Warning, TEXT("Roi tensor must be a 1-D tensor."));
+			return false;
+		}
+		if(InputShapes[1].GetData()[0] != InputShapes[0].Rank() * 2)
+		{
+			UE_LOG(LogNNE, Warning, TEXT("Roi tensor must have dimension 2*N (where N is the input rank)."));
+			return false;
+		}
+		if(InputTypes.Num() == 3)
+		{
+			if(InputTypes[2] != ENNETensorDataType::Float)
+			{
+				UE_LOG(LogNNE, Warning, TEXT("Scales tensor must be of type float."));
+				return false;
+			}
+			if(InputShapes[2].Rank() != 1)
+			{
+				UE_LOG(LogNNE, Warning, TEXT("Scales tensor must be a 1-D tensor."));
+				return false;
+			}
+			if(InputShapes[2].GetData()[0] != InputShapes[0].Rank())
+			{
+				UE_LOG(LogNNE, Warning, TEXT("Scales tensor must have dimension N (where N is the input rank)."));
+				return false;
+			}
+		}
+		else
+		{
+			if(InputTypes[2] != ENNETensorDataType::None)
+			{
+				UE_LOG(LogNNE, Warning, TEXT("Scales tensor must be empty (i.e. empty name and data type 'None') when Sizes is specified."));
+				return false;
+			}
+			if(InputTypes[3] != ENNETensorDataType::Int64)
+			{
+				UE_LOG(LogNNE, Warning, TEXT("Sizes tensor must be of type Int64."));
+				return false;
+			}
+			if(InputShapes[3].Rank() != 1)
+			{
+				UE_LOG(LogNNE, Warning, TEXT("Sizes tensor must be a 1-D tensor."));
+				return false;
+			}
+			if(InputShapes[3].GetData()[0] != InputShapes[0].Rank())
+			{
+				UE_LOG(LogNNE, Warning, TEXT("Sizes tensor must have dimension N (where N is the input rank)."));
+				return false;
+			}
+		}
+
+		// NOTE: the current RDG Hlsl infrastructure does NOT support empty tensors (FModelInstance::Init() breaks), 
+		// therefore it's not possible to specify the Sizes tensor (as Scales would need to be empty).
+		if(InputTypes.Num() == 4)
+		{
+			UE_LOG(LogNNE, Warning, TEXT("Sizes tensor is not currently supported, please use Scales tensor."));
+			return false;
+		}
+
+		return bIsValid;
+	}
+
+	FOperatorHlsl* CreateResizeOperator()
+	{
+		return new FResize();
+	}
+
+	bool RegisterResizeOperator(FOperatorRegistryHlsl& Registry)
+	{
+		// Note: support of a particular version is partial with respect to tensor data types (only the most typical ones are usually supported).
+		Registry.OpAdd({{TEXT("Resize"), TEXT("Onnx")}, 10}, CreateResizeOperator, ValidateResizeOperator);
+		Registry.OpAdd({{TEXT("Resize"), TEXT("Onnx")}, 11}, CreateResizeOperator, ValidateResizeOperator);
+		Registry.OpAdd({{TEXT("Resize"), TEXT("Onnx")}, 13}, CreateResizeOperator, ValidateResizeOperator);
+		Registry.OpAdd({{TEXT("Resize"), TEXT("Onnx")}, 18}, CreateResizeOperator, ValidateResizeOperator);
+		Registry.OpAdd({{TEXT("Resize"), TEXT("Onnx")}, 19}, CreateResizeOperator, ValidateResizeOperator);
+		return true;
+	}
+} // UE::NNERuntimeRDG::Private::Hlsl
diff --git a/Engine/Plugins/Experimental/NNERuntimeRDG/Source/NNERuntimeRDG/Private/Hlsl/NNERuntimeRDGResize.h b/Engine/Plugins/Experimental/NNERuntimeRDG/Source/NNERuntimeRDG/Private/Hlsl/NNERuntimeRDGResize.h
new file mode 100644
--- /dev/null
+++ b/Engine/Plugins/Experimental/NNERuntimeRDG/Source/NNERuntimeRDG/Private/Hlsl/NNERuntimeRDGResize.h
@@ -0,0 +1,11 @@
+// Copyright Epic Games, Inc. All Rights Reserved.
+
+#pragma once
+
+#include "NNERuntimeRDGBase.h"
+#include "NNERuntimeRDGHlslOp.h"
+
+namespace UE::NNERuntimeRDG::Private::Hlsl
+{
+	bool RegisterResizeOperator(FOperatorRegistryHlsl& Registry);
+} // UE::NNERuntimeRDG::Private::Hlsl
diff --git a/Engine/Plugins/Experimental/NNERuntimeRDG/Source/NNERuntimeRDG/Private/NNERuntimeRDGHlsl.cpp b/Engine/Plugins/Experimental/NNERuntimeRDG/Source/NNERuntimeRDG/Private/NNERuntimeRDGHlsl.cpp
--- a/Engine/Plugins/Experimental/NNERuntimeRDG/Source/NNERuntimeRDG/Private/NNERuntimeRDGHlsl.cpp
+++ b/Engine/Plugins/Experimental/NNERuntimeRDG/Source/NNERuntimeRDG/Private/NNERuntimeRDGHlsl.cpp
@@ -35,6 +35,7 @@
 #include "Hlsl/NNERuntimeRDGPad.h"
 #include "Hlsl/NNERuntimeRDGPool.h"
 #include "Hlsl/NNERuntimeRDGReduce.h"
+#include "Hlsl/NNERuntimeRDGResize.h"
 #include "Hlsl/NNERuntimeRDGReshape.h"
 #include "Hlsl/NNERuntimeRDGScatterND.h"
 #include "Hlsl/NNERuntimeRDGShape.h"
@@ -78,6 +79,7 @@ bool UNNERuntimeRDGHlslImpl::Init()
 	RegisterPoolOperators(*Registry);
 	RegisterReduceOperators(*Registry);
 	RegisterReshapeOperator(*Registry);
+	RegisterResizeOperator(*Registry);
 	RegisterScatterNDOperator(*Registry);
 	RegisterShapeOperator(*Registry);
 	RegisterSizeOperator(*Registry);
