NNE - Add Hlsl ReduceSum-13 and ReduceMean-18

Problems:
* RDG Hlsl ReduceSum-13 and ReduceMean-18 implementations are missing.

Solutions:
* Add ReduceSum-13 implementation.
* Add ReduceMean-18 implementation.

Tests:
* Win64 Standalone tests: all
* Win64 Editor tests: all except some OIDN tests which are failing regardless of the changes
* Note: cannot test case noop_with_empty_axes=1 because ORT cpu erroneously asserts when importing tests.

#rnx

#virtualized

diff --git a/Engine/Plugins/Experimental/NNERuntimeRDG/Source/NNERuntimeRDG/Private/Hlsl/NNERuntimeRDGReduce.cpp b/Engine/Plugins/Experimental/NNERuntimeRDG/Source/NNERuntimeRDG/Private/Hlsl/NNERuntimeRDGReduce.cpp
--- a/Engine/Plugins/Experimental/NNERuntimeRDG/Source/NNERuntimeRDG/Private/Hlsl/NNERuntimeRDGReduce.cpp
+++ b/Engine/Plugins/Experimental/NNERuntimeRDG/Source/NNERuntimeRDG/Private/Hlsl/NNERuntimeRDGReduce.cpp
@@ -3,6 +3,7 @@
 #include "NNERuntimeRDGReduce.h"
 #include "NNEHlslShadersReduceCS.h"
 #include "NNERuntimeRDGHlslHelper.h"
+#include "Helper/NNERuntimeRDGOperatorHelper.h"
 #include "NNETensor.h"
 #include "NNETypes.h"
 #include "RenderGraphUtils.h"
@@ -14,7 +15,7 @@ namespace UE::NNERuntimeRDG::Private::Hlsl
 	/**
 	 * Reduce operators implementation
 	 */
-	template< UE::NNEHlslShaders::Internal::EReduceOperatorType ReduceOperatorType >
+	template< UE::NNEHlslShaders::Internal::EReduceOperatorType ReduceOperatorType, int Version>
 	class FReduceOperator : public FOperatorHlsl
 	{
 
@@ -23,32 +24,76 @@ namespace UE::NNERuntimeRDG::Private::Hlsl
 		FReduceOperator() = default;
 		virtual ~FReduceOperator() = default;
 
+		static constexpr bool bAxesAsInput = 
+			(Version >= 13 && ReduceOperatorType == UE::NNEHlslShaders::Internal::EReduceOperatorType::Sum) ||
+			(Version >= 18 && ReduceOperatorType == UE::NNEHlslShaders::Internal::EReduceOperatorType::Average)
+			;
+
 	private:
 
-		TArray<int32> Axes;
+		TArray<int32, TInlineAllocator<NNE::FTensorShape::MaxRank>> Axes;
 		int32 KeepDims = 1;
 
 	public:
 
 		virtual int PrepareOutputs(TConstArrayView<NNE::Internal::FTensorRef> InputTensors, TArrayView<NNE::Internal::FTensorRef> OutputTensors) override
 		{
-			check(InputTensors.Num() == 1);
+			check(bAxesAsInput ? (InputTensors.Num() >= 1 && InputTensors.Num() <= 2) : InputTensors.Num() == 1);
 			check(OutputTensors.Num() == 1);
 
+			if constexpr (bAxesAsInput)
+			{
+				if(InputTensors.Num() == 2)
+				{
+					const NNE::Internal::FTensorRef AxesTensor = InputTensors[1];
+					if(!AxesTensor->HasPreparedData())
+					{
+						UE_LOG(LogNNE, Warning, TEXT("Reduce tensor `axes` (name: %s) must be CPU constant."), *InputTensors[1]->GetName());
+						return false;
+					}
+					if(AxesTensor->GetVolume() != 0)
+					{
+						OperatorHelper::GetInt32ArrayFromConstTensor(Axes, AxesTensor);
+					}
+				}
+			}
+
 			const NNE::FTensorShape& InputShape = InputTensors[0]->GetShape();
 			const int32 InputRank = InputShape.Rank();
-			TArray<uint32> OutputShape;
 
-			OutputShape.Reserve(InputRank);
-			for (int i = 0; i < InputRank; ++i)
+			for (int32& Axis : Axes)
 			{
-				if (!Axes.Contains(i))
+				if (Axis > InputRank || Axis  < -InputRank)
 				{
-					OutputShape.Add(InputShape.GetData()[i]);
+					UE_LOG(LogNNE, Warning, TEXT("Reduce operators 'Axes' attribute should contain value be in the range [-r,r] with r being the rank of the input (name: %s) however got %d while rank is %d."), *InputTensors[0]->GetName(), Axis, InputRank);
+					return false;
 				}
-				else if (KeepDims)
+
+				if (Axis  < 0)
 				{
-					OutputShape.Add(1);
+					Axis += InputRank;
+				}
+			}
+			
+			TArray<uint32> OutputShape;
+
+			OutputShape.Reserve(InputRank);
+			if(bAxesAsInput && Axes.IsEmpty())
+			{
+				OutputShape = InputShape.GetData();
+			}
+			else
+			{
+				for (int i = 0; i < InputRank; ++i)
+				{
+					if (!Axes.Contains(i))
+					{
+						OutputShape.Add(InputShape.GetData()[i]);
+					}
+					else if (KeepDims)
+					{
+						OutputShape.Add(1);
+					}
 				}
 			}
 
@@ -59,36 +104,30 @@ namespace UE::NNERuntimeRDG::Private::Hlsl
 
 		virtual bool Initialize(TConstArrayView<NNE::FTensorDesc> InputTensorDescs, TConstArrayView<NNE::FTensorDesc> OutputTensorDescs, const NNE::FAttributeMap& Attributes) override
 		{
-			check(InputTensorDescs.Num() == 1);
+			check(bAxesAsInput ? (InputTensorDescs.Num() >= 1 && InputTensorDescs.Num() <= 2) : InputTensorDescs.Num() == 1);
 			check(OutputTensorDescs.Num() == 1);
 
 			const int32 InputRank = InputTensorDescs[0].GetShape().Rank();
-			TArray<int32> AxesTemp;
-
-			for (int i = 0; i < InputRank; ++i)
-			{
-				AxesTemp.Add(i);
-			}
-			AxesTemp = Attributes.GetValueOrDefault<TArray<int32>>(TEXT("axes"), AxesTemp);
-			AxesTemp.Sort();
-
-			Axes.Empty();
+			
+			TArray<int32> AxesDefault;
+			AxesDefault.Reserve(NNE::FTensorShape::MaxRank);
 
-			for (int32 axis : AxesTemp)
+			if(!bAxesAsInput || Attributes.GetValueOrDefault<int>(TEXT("noop_with_empty_axes"), 0) != 1)
 			{
-				if (axis > InputRank || axis  < -InputRank)
+				for (int i = 0; i < InputRank; ++i)
 				{
-					UE_LOG(LogNNE, Warning, TEXT("Reduce operators 'Axes' attribute should contain value be in the range [-r,r] with r being the rank of the input (name: %s) however got %d while rank is %d."), *InputTensorDescs[0].GetName(), axis, InputRank);
-					return false;
+					AxesDefault.Add(i);
 				}
+			}
+			
+			Axes = Attributes.GetValueOrDefault<TArray<int32>>(TEXT("axes"), AxesDefault);
+			Axes.Sort();
 
-				if (axis  < 0)
-				{
-					axis = InputRank + axis;
-				}
-				Axes.AddUnique(axis);
+			if (!bAxesAsInput && Axes.Num() == 0)
+			{
+				UE_LOG(LogNNE, Warning, TEXT("Reduce attribute `axes` cannot be empty."));
+				return false;
 			}
-			check(Axes.Num() > 0);
 
 			KeepDims = Attributes.GetValueOrDefault<int32>(TEXT("keepdims"), 1);
 
@@ -99,11 +138,10 @@ namespace UE::NNERuntimeRDG::Private::Hlsl
 		{
 			using namespace UE::NNEHlslShaders::Internal;
 
-			check(InputTensors.Num() == 1);
+			check(bAxesAsInput ? (InputTensors.Num() >= 1 && InputTensors.Num() <= 2) : InputTensors.Num() == 1);
 			check(OutputTensors.Num() == 1);
 			check(InputTensors[0] != FTensorRDGRef{});
 			check(OutputTensors[0] != FTensorRDGRef{});
-			check(Axes.Num() > 0);
 
 			const FTensorRDG& Input = *InputTensors[0];
 			const FTensorRDG& Output = *OutputTensors[0];
@@ -112,82 +150,106 @@ namespace UE::NNERuntimeRDG::Private::Hlsl
 			RDG_EVENT_SCOPE(GraphBuilder, "NNE.Operator.Hlsl.Reduce");
 			RDG_GPU_STAT_SCOPE(GraphBuilder, FNNEOperatorReduce);
 
-			FRDGBufferRef CurrInput = Input.GetBuffer();
-			TArray<uint32, TInlineAllocator<NNE::FTensorShape::MaxRank>> CurrInputShape(Input.GetShape().GetData());
-			for (int32 i = Axes.Num()-1; i >= 0; --i)
+			if(Axes.IsEmpty())
 			{
-				const int32 Axis = Axes[i];
-
-				TReduceCS::FParameters* Parameters = GraphBuilder.AllocParameters<TReduceCS::FParameters>();
-				TReduceCS::FillInParameters(CurrInputShape, Axis, Parameters);
-
-				FRDGBufferRef CurrOutput = nullptr;
-				if (i != 0)
-				{
-					const FRDGBufferDesc TempBufferDesc = FRDGBufferDesc::CreateBufferDesc(Output.GetElementByteSize(), Parameters->NumElemBeforeAxis * Parameters->NumElemAfterAxis);
-					CurrOutput = GraphBuilder.CreateBuffer(TempBufferDesc, TEXT("NNE.Operator.Hlsl.Reduce.TempBuffer"), ERDGBufferFlags::None);
-				}
-				else
+				AddCopyBufferPass(GraphBuilder, Output.GetBuffer(), Input.GetBuffer());
+			}
+			else
+			{
+				FRDGBufferRef CurrInput = Input.GetBuffer();
+				TArray<uint32, TInlineAllocator<NNE::FTensorShape::MaxRank>> CurrInputShape(Input.GetShape().GetData());
+				for (int32 i = Axes.Num()-1; i >= 0; --i)
 				{
-					CurrOutput = Output.GetBuffer();
+					const int32 Axis = Axes[i];
+
+					TReduceCS::FParameters* Parameters = GraphBuilder.AllocParameters<TReduceCS::FParameters>();
+					TReduceCS::FillInParameters(CurrInputShape, Axis, Parameters);
+
+					FRDGBufferRef CurrOutput = nullptr;
+					if (i != 0)
+					{
+						const FRDGBufferDesc TempBufferDesc = FRDGBufferDesc::CreateBufferDesc(Output.GetElementByteSize(), Parameters->NumElemBeforeAxis * Parameters->NumElemAfterAxis);
+						CurrOutput = GraphBuilder.CreateBuffer(TempBufferDesc, TEXT("NNE.Operator.Hlsl.Reduce.TempBuffer"), ERDGBufferFlags::None);
+					}
+					else
+					{
+						CurrOutput = Output.GetBuffer();
+					}
+					check(CurrOutput);
+
+					TReduceCS::EnqueueRDG(GraphBuilder, Parameters, CurrInput, CurrOutput, ReduceOperatorType);
+					CurrInput = CurrOutput;
+					CurrInputShape[i] = 1;
 				}
-				check(CurrOutput);
-
-				TReduceCS::EnqueueRDG(GraphBuilder, Parameters, CurrInput, CurrOutput, ReduceOperatorType);
-				CurrInput = CurrOutput;
-				CurrInputShape[i] = 1;
 			}
 		}
 	};
 
+	template< UE::NNEHlslShaders::Internal::EReduceOperatorType ReduceOperatorType, int Version>
 	bool ValidateReduceOperator(const NNE::FAttributeMap& AttributeMap, TConstArrayView<ENNETensorDataType> InputTypes, TConstArrayView<NNE::FSymbolicTensorShape> InputShapes)
 	{
 		bool bIsValid = true;
 
 		// This match versions 1 and 11 of the Reduces operators, next versions are 13
 		// ReduceOperator are ReduceL1, ReduceL2, ReduceLogSum, ReduceLogSumExp, ReduceMax, ReduceMean, ReduceMin, ReduceProd, ReduceSum, ReduceSumSquare
+		// ReduceMean-13 is also supported
 		// [REDACTED_PROFILE_URL]/onnx/blob/main/docs/Changelog.md#Reduce-1
 		FAttributeValidator AttributeValidator;
-		AttributeValidator.AddOptional(TEXT("axes"), ENNEAttributeDataType::Int32Array);
 		AttributeValidator.AddOptional(TEXT("keepdims"), ENNEAttributeDataType::Int32);
+		if constexpr (!FReduceOperator<ReduceOperatorType, Version>::bAxesAsInput)
+		{
+			AttributeValidator.AddOptional(TEXT("axes"), ENNEAttributeDataType::Int32Array);
+		}
+		else
+		{
+			AttributeValidator.AddOptional(TEXT("noop_with_empty_axes"), ENNEAttributeDataType::Int32);
+		}
 		bIsValid &= AttributeValidator.Validate(AttributeMap);
 
 		FInputValidator InputValidator;
+		InputValidator.SetTemplateCount(2);
 		InputValidator.AddSupportedType(ENNETensorDataType::Float);
+		InputValidator.AddSupportedType(ENNETensorDataType::Int64, /*TemplateIdx*/ 1);
 		InputValidator.AddRequired();
+		if constexpr (FReduceOperator<ReduceOperatorType, Version>::bAxesAsInput)
+		{
+			InputValidator.AddOptional(/*TemplateIdx*/ 1);
+		}
 		bIsValid &= InputValidator.Validate(InputTypes);
 
 		return bIsValid;
 	}
 
-	template< UE::NNEHlslShaders::Internal::EReduceOperatorType ReduceOperatorType >
+	template< UE::NNEHlslShaders::Internal::EReduceOperatorType ReduceOperatorType, int Version>
 	FOperatorHlsl* CreateReduceOperator()
 	{
-		return new FReduceOperator<ReduceOperatorType>();
+		return new FReduceOperator<ReduceOperatorType, Version>();
 	}
 
 	bool RegisterReduceOperators(FOperatorRegistryHlsl& Registry)
 	{
+		using namespace UE::NNEHlslShaders::Internal;
 		//ReduceLogSum not yet supported as multi axis case require to apply all reduction first then log.
 		//ReduceSumSquare not yet supported as multi axis case require to apply square on whole tensor first then sum.
-
-		Registry.OpAdd({ {TEXT("ReduceL1"),        TEXT("Onnx")},  1  }, CreateReduceOperator<UE::NNEHlslShaders::Internal::EReduceOperatorType::L1>, ValidateReduceOperator);
-		Registry.OpAdd({ {TEXT("ReduceL2"),        TEXT("Onnx")},  1  }, CreateReduceOperator<UE::NNEHlslShaders::Internal::EReduceOperatorType::L2>, ValidateReduceOperator);
-		Registry.OpAdd({ {TEXT("ReduceLogSumExp"), TEXT("Onnx")},  1  }, CreateReduceOperator<UE::NNEHlslShaders::Internal::EReduceOperatorType::LogSumExp>, ValidateReduceOperator);
-		Registry.OpAdd({ {TEXT("ReduceMax"),       TEXT("Onnx")},  1  }, CreateReduceOperator<UE::NNEHlslShaders::Internal::EReduceOperatorType::Max>, ValidateReduceOperator);
-		Registry.OpAdd({ {TEXT("ReduceMean"),      TEXT("Onnx")},  1  }, CreateReduceOperator<UE::NNEHlslShaders::Internal::EReduceOperatorType::Average>, ValidateReduceOperator);
-		Registry.OpAdd({ {TEXT("ReduceMin"),       TEXT("Onnx")},  1  }, CreateReduceOperator<UE::NNEHlslShaders::Internal::EReduceOperatorType::Min>, ValidateReduceOperator);
-		Registry.OpAdd({ {TEXT("ReduceProd"),      TEXT("Onnx")},  1  }, CreateReduceOperator<UE::NNEHlslShaders::Internal::EReduceOperatorType::Prod>, ValidateReduceOperator);
-		Registry.OpAdd({ {TEXT("ReduceSum"),       TEXT("Onnx")},  1  }, CreateReduceOperator<UE::NNEHlslShaders::Internal::EReduceOperatorType::Sum>, ValidateReduceOperator);
-
-		Registry.OpAdd({ {TEXT("ReduceL1"),        TEXT("Onnx")}, 11 }, CreateReduceOperator<UE::NNEHlslShaders::Internal::EReduceOperatorType::L1>, ValidateReduceOperator);
-		Registry.OpAdd({ {TEXT("ReduceL2"),        TEXT("Onnx")}, 11 }, CreateReduceOperator<UE::NNEHlslShaders::Internal::EReduceOperatorType::L2>, ValidateReduceOperator);
-		Registry.OpAdd({ {TEXT("ReduceLogSumExp"), TEXT("Onnx")}, 11 }, CreateReduceOperator<UE::NNEHlslShaders::Internal::EReduceOperatorType::LogSumExp>, ValidateReduceOperator);
-		Registry.OpAdd({ {TEXT("ReduceMax"),       TEXT("Onnx")}, 11 }, CreateReduceOperator<UE::NNEHlslShaders::Internal::EReduceOperatorType::Max>, ValidateReduceOperator);
-		Registry.OpAdd({ {TEXT("ReduceMean"),      TEXT("Onnx")}, 11 }, CreateReduceOperator<UE::NNEHlslShaders::Internal::EReduceOperatorType::Average>, ValidateReduceOperator);
-		Registry.OpAdd({ {TEXT("ReduceMin"),       TEXT("Onnx")}, 11 }, CreateReduceOperator<UE::NNEHlslShaders::Internal::EReduceOperatorType::Min>, ValidateReduceOperator);
-		Registry.OpAdd({ {TEXT("ReduceProd"),      TEXT("Onnx")}, 11 }, CreateReduceOperator<UE::NNEHlslShaders::Internal::EReduceOperatorType::Prod>, ValidateReduceOperator);
-		Registry.OpAdd({ {TEXT("ReduceSum"),       TEXT("Onnx")}, 11 }, CreateReduceOperator<UE::NNEHlslShaders::Internal::EReduceOperatorType::Sum>, ValidateReduceOperator);
+		#define OPS(Version) \
+		Registry.OpAdd({ {TEXT("ReduceL1"),        TEXT("Onnx")}, Version}, CreateReduceOperator<EReduceOperatorType::L1,Version>, ValidateReduceOperator<EReduceOperatorType::L1,Version>); \
+		Registry.OpAdd({ {TEXT("ReduceL2"),        TEXT("Onnx")}, Version}, CreateReduceOperator<EReduceOperatorType::L2,Version>, ValidateReduceOperator<EReduceOperatorType::L2,Version>); \
+		Registry.OpAdd({ {TEXT("ReduceLogSumExp"), TEXT("Onnx")}, Version}, CreateReduceOperator<EReduceOperatorType::LogSumExp,Version>, ValidateReduceOperator<EReduceOperatorType::LogSumExp,Version>); \
+		Registry.OpAdd({ {TEXT("ReduceMax"),       TEXT("Onnx")}, Version}, CreateReduceOperator<EReduceOperatorType::Max,Version>, ValidateReduceOperator<EReduceOperatorType::Max,Version>); \
+		Registry.OpAdd({ {TEXT("ReduceMean"),      TEXT("Onnx")}, Version}, CreateReduceOperator<EReduceOperatorType::Average,Version>, ValidateReduceOperator<EReduceOperatorType::Average,Version>); \
+		Registry.OpAdd({ {TEXT("ReduceMin"),       TEXT("Onnx")}, Version}, CreateReduceOperator<EReduceOperatorType::Min,Version>, ValidateReduceOperator<EReduceOperatorType::Min,Version>); \
+		Registry.OpAdd({ {TEXT("ReduceProd"),      TEXT("Onnx")}, Version}, CreateReduceOperator<EReduceOperatorType::Prod,Version>, ValidateReduceOperator<EReduceOperatorType::Prod,Version>); \
+		Registry.OpAdd({ {TEXT("ReduceSum"),       TEXT("Onnx")}, Version}, CreateReduceOperator<EReduceOperatorType::Sum,Version>, ValidateReduceOperator<EReduceOperatorType::Sum,Version>);
+
+		OPS(1)
+		OPS(11)
+		
+		Registry.OpAdd({ {TEXT("ReduceSum"),       TEXT("Onnx")}, 13 }, CreateReduceOperator<EReduceOperatorType::Sum, 13>, ValidateReduceOperator<EReduceOperatorType::Sum, 13>);
+		Registry.OpAdd({ {TEXT("ReduceMean"),       TEXT("Onnx")}, 13 }, CreateReduceOperator<EReduceOperatorType::Average, 18>, ValidateReduceOperator<EReduceOperatorType::Average, 13>);
+		
+		Registry.OpAdd({ {TEXT("ReduceMean"),       TEXT("Onnx")}, 18 }, CreateReduceOperator<EReduceOperatorType::Average, 18>, ValidateReduceOperator<EReduceOperatorType::Average, 18>);
+
+		#undef OPS
 
 		return true;
 	}
