NNEDenoiser: Fix minor shader warning

#rnx

diff --git a/Engine/Plugins/Experimental/NNERuntimeRDG/Source/NNEHlslShaders/Shaders/NNEHlslShadersResize.usf b/Engine/Plugins/Experimental/NNERuntimeRDG/Source/NNEHlslShaders/Shaders/NNEHlslShadersResize.usf
--- a/Engine/Plugins/Experimental/NNERuntimeRDG/Source/NNEHlslShaders/Shaders/NNEHlslShadersResize.usf
+++ b/Engine/Plugins/Experimental/NNERuntimeRDG/Source/NNEHlslShaders/Shaders/NNEHlslShadersResize.usf
@@ -189,7 +189,7 @@ void Resize(in const uint3 DispatchThreadID : SV_DispatchThreadID)
 
 				STATIC_LOOP(DimIdx, 0, NUM_DIMENSIONS)
 				{
-					uint PointCoord = (PointIdx & (1 << DimIdx)) >> DimIdx;
+					uint PointCoord = (PointIdx & (1u << DimIdx)) >> DimIdx;
 					PointCoords[DimIdx] = InputIterator[DimIdx] + PointCoord;
 					if(PointCoords[DimIdx] >= (InputTensorInfo[DimIdx][SIZE_IDX] - 1))
 					{
