NNE - Refactor handling of constant tensors in Hlsl backend to support zero-sized tensors

Problems:
* RDG buffers can't be zero-sized but zero-sized tensors are perfectly valid in NNE format. The implementation tries to register an external zero-sized buffer for a FTensorRDG, leading to a crash.
* FTensorRDG::HasPreparedData() is sometimes erroneously used to query whether the tensor is constant.
* When generating an ONNX from a JSON specification of a zero-sized tensor, the tensor is not registered as initializer. This is because (similarly as before) the implementation checks whether it has data (non-null Data ptr) to assess whether it's constant.

Solutions:
* Don't register an external buffer when tensor is zero-sized.
* Use FTensorRDG::IsConstant() to check whether the tensor is constant. It returns true when either the tensor has data or it's zero-sized.
* Refactor model builders' API to distinguish between "normal" tensors and constant ones. No need to check data pointer to assess constantness.

Tests:
* Test via Resize operator by specifying zero-sized ROI weight tensor.
* Win64 Standalone tests: all
* Win64 Editor tests: all

#rnx

#virtualized

diff --git a/Engine/Plugins/Experimental/NNERuntimeRDG/Source/NNERuntimeRDG/Private/Helper/NNERuntimeRDGOperatorHelper.cpp b/Engine/Plugins/Experimental/NNERuntimeRDG/Source/NNERuntimeRDG/Private/Helper/NNERuntimeRDGOperatorHelper.cpp
--- a/Engine/Plugins/Experimental/NNERuntimeRDG/Source/NNERuntimeRDG/Private/Helper/NNERuntimeRDGOperatorHelper.cpp
+++ b/Engine/Plugins/Experimental/NNERuntimeRDG/Source/NNERuntimeRDG/Private/Helper/NNERuntimeRDGOperatorHelper.cpp
@@ -10,7 +10,7 @@ namespace UE::NNERuntimeRDG::Private::OperatorHelper
 		check(Tensor != nullptr);
 		Attr.Reset();
 
-		if (!Tensor->HasPreparedData())
+		if (!Tensor->IsConstant())
 		{
 			return false;
 		}
diff --git a/Engine/Plugins/Experimental/NNERuntimeRDG/Source/NNERuntimeRDG/Private/Hlsl/NNERuntimeRDGResize.cpp b/Engine/Plugins/Experimental/NNERuntimeRDG/Source/NNERuntimeRDG/Private/Hlsl/NNERuntimeRDGResize.cpp
--- a/Engine/Plugins/Experimental/NNERuntimeRDG/Source/NNERuntimeRDG/Private/Hlsl/NNERuntimeRDGResize.cpp
+++ b/Engine/Plugins/Experimental/NNERuntimeRDG/Source/NNERuntimeRDG/Private/Hlsl/NNERuntimeRDGResize.cpp
@@ -77,15 +77,6 @@ namespace UE::NNERuntimeRDG::Private::Hlsl
 				return -1;
 			}
 
-			for(int Idx = 1; Idx < InputTensors.Num(); ++Idx)
-			{
-				if(!InputTensors[Idx]->HasPreparedData())
-				{
-					UE_LOG(LogNNE, Warning, TEXT("Hlsl Resize: tensor %s should be constant for shape inference to succeed."), *InputTensors[Idx]->GetName());
-					return -1;
-				}
-			}
-
 			const NNE::Internal::FTensor& Input = *InputTensors[0];
 			const NNE::Internal::FTensor& Roi = *InputTensors[1];
 
@@ -107,25 +98,35 @@ namespace UE::NNERuntimeRDG::Private::Hlsl
 				return -1;
 			}
 
-			if(Roi.GetDataType() != ENNETensorDataType::Float)
+			if(CoordTransMode == UE::NNEHlslShaders::Internal::ECoordTransMode::TfCropAndResize
+				&& Roi.GetShape().Volume() != 2 * Input.GetShape().Rank())
 			{
-				UE_LOG(LogNNE, Warning, TEXT("Hlsl Resize: only float data type supported for`roi` tensor (name %s)."), *Roi.GetName());
+				UE_LOG(LogNNE, Warning, TEXT("Hlsl Resize: `roi` tensor (name %s) must have 2 * N length."), *Roi.GetName());
 				return -1;
 			}
 
-			if(!Roi.HasPreparedData())
+			if(Roi.GetShape().Volume() != 0)
 			{
-				UE_LOG(LogNNE, Warning, TEXT("Hlsl Resize: `roi` tensor could not be made constant. (name %s)."), *Roi.GetName());
-				return -1;
-			}
-			
-			RegionOfInterest = Roi.GetPreparedData<float>();
+				if(Roi.GetDataType() != ENNETensorDataType::Float)
+				{
+					UE_LOG(LogNNE, Warning, TEXT("Hlsl Resize: only float data type supported for`roi` tensor (name %s)."), *Roi.GetName());
+					return -1;
+				}
 
-			if(CoordTransMode == UE::NNEHlslShaders::Internal::ECoordTransMode::TfCropAndResize
-				&& RegionOfInterest.Num() != 2 * Input.GetShape().Rank())
-			{
-				UE_LOG(LogNNE, Warning, TEXT("Hlsl Resize: `roi` tensor (name %s) must have 2 * N length."), *Roi.GetName());
-				return -1;
+				if(!Roi.HasPreparedData())
+				{
+					UE_LOG(LogNNE, Warning, TEXT("Hlsl Resize: `roi` tensor could not be made constant. (name %s)."), *Roi.GetName());
+					return -1;
+				}
+				
+				RegionOfInterest = Roi.GetPreparedData<float>();
+
+				if(CoordTransMode == UE::NNEHlslShaders::Internal::ECoordTransMode::TfCropAndResize
+					&& RegionOfInterest.Num() != 2 * Input.GetShape().Rank())
+				{
+					UE_LOG(LogNNE, Warning, TEXT("Hlsl Resize: `roi` tensor (name %s) must have 2 * N length."), *Roi.GetName());
+					return -1;
+				}
 			}
 
 			TArray<uint32> OutputShapeData;
@@ -364,6 +365,8 @@ namespace UE::NNERuntimeRDG::Private::Hlsl
 
 	bool ValidateResizeOperator(const NNE::FAttributeMap& AttributeMap, TConstArrayView<ENNETensorDataType> InputTypes, TConstArrayView<NNE::FSymbolicTensorShape> InputShapes)
 	{
+		using namespace UE::NNEHlslShaders::Internal;
+		
 		bool bIsValid = true;
 
 		FAttributeValidator AttributeValidator;
@@ -398,11 +401,13 @@ namespace UE::NNERuntimeRDG::Private::Hlsl
 			UE_LOG(LogNNE, Warning, TEXT("Roi tensor must be a 1-D tensor."));
 			return false;
 		}
-		if(InputShapes[1].GetData()[0] != InputShapes[0].Rank() * 2)
+		if(FResizeCS::CoordTransModeFromString(*AttributeMap.GetValueOrDefault<FString>(TEXT("coordinate_transformation_mode"), TEXT("half_pixel"))) == ECoordTransMode::TfCropAndResize &&
+			InputShapes[1].GetData()[0] != 2 * InputShapes[0].Rank())
 		{
-			UE_LOG(LogNNE, Warning, TEXT("Roi tensor must have dimension 2*N (where N is the input rank)."));
+			UE_LOG(LogNNE, Warning, TEXT("Roi tensor must have dimension 2*N (where N is the input rank) when `coordinate_transformation_mode` is `tf_crop_and_resize`."));
 			return false;
 		}
+
 		if(InputTypes.Num() == 3)
 		{
 			if(InputTypes[2] != ENNETensorDataType::Float)
diff --git a/Engine/Plugins/Experimental/NNERuntimeRDG/Source/NNERuntimeRDG/Private/NNERuntimeRDGBase.h b/Engine/Plugins/Experimental/NNERuntimeRDG/Source/NNERuntimeRDG/Private/NNERuntimeRDGBase.h
--- a/Engine/Plugins/Experimental/NNERuntimeRDG/Source/NNERuntimeRDG/Private/NNERuntimeRDGBase.h
+++ b/Engine/Plugins/Experimental/NNERuntimeRDG/Source/NNERuntimeRDG/Private/NNERuntimeRDGBase.h
@@ -64,6 +64,7 @@ public:
 		return TensorRDG;
 	}
 
+	bool IsValid() const { return HasBuffer() || GetVolume() == 0; }
 	bool HasBuffer() const { return Buffer != FRDGBufferRef{}; }
 	void SetBuffer(FRDGBufferRef Inbuffer){ Buffer = Inbuffer; }
 	FRDGBufferRef GetBuffer() const { return Buffer; }
diff --git a/Engine/Plugins/Experimental/NNERuntimeRDG/Source/NNERuntimeRDG/Private/NNERuntimeRDGModel.cpp b/Engine/Plugins/Experimental/NNERuntimeRDG/Source/NNERuntimeRDG/Private/NNERuntimeRDGModel.cpp
--- a/Engine/Plugins/Experimental/NNERuntimeRDG/Source/NNERuntimeRDG/Private/NNERuntimeRDGModel.cpp
+++ b/Engine/Plugins/Experimental/NNERuntimeRDG/Source/NNERuntimeRDG/Private/NNERuntimeRDGModel.cpp
@@ -230,7 +230,7 @@ FModelInstanceRDG::EEnqueueRDGStatus FModelInstanceRDG::EnqueueRDG(FRDGBuilder&
 	//Create temporary buffers for NOT const intermediate tensors
 	for (FTensorRDG& TensorRDG : IntermediateTensorRDGs)
 	{
-		if (!TensorRDG.HasPreparedData())
+		if (!TensorRDG.IsConstant())
 		{
 			const FRDGBufferDesc BufferDesc = CreateRDGBufferDescForTensorRDG(TensorRDG);
 			const FRDGBufferRef TensorBuffer = RDGBuilder.CreateBuffer(BufferDesc, TEXT("NNE.Tensor.Intermediate"), ERDGBufferFlags::None);
@@ -246,7 +246,7 @@ FModelInstanceRDG::EEnqueueRDGStatus FModelInstanceRDG::EnqueueRDG(FRDGBuilder&
 		checkCode(
 			for (const TPair<int32, FTensorRDGRef>& TensorRDG : AllTensorRDGRefs) 
 			{ 
-				check(TensorRDG.Value->GetBuffer() != nullptr); 
+				check(TensorRDG.Value->IsValid()); 
 			}
 		);
 	}
diff --git a/Engine/Plugins/Experimental/NNERuntimeRDG/Source/NNERuntimeRDG/Private/NNERuntimeRDGModelHlsl.cpp b/Engine/Plugins/Experimental/NNERuntimeRDG/Source/NNERuntimeRDG/Private/NNERuntimeRDGModelHlsl.cpp
--- a/Engine/Plugins/Experimental/NNERuntimeRDG/Source/NNERuntimeRDG/Private/NNERuntimeRDGModelHlsl.cpp
+++ b/Engine/Plugins/Experimental/NNERuntimeRDG/Source/NNERuntimeRDG/Private/NNERuntimeRDGModelHlsl.cpp
@@ -97,8 +97,16 @@ bool FModelInstance::PrepareModelRDG(FRDGBuilder& RDGBuilder)
 	{
 		const TRefCountPtr<FRDGPooledBuffer>& PooledBuffer = WeightsExternalRDGResources[Idx];
 		FTensorRDG& Tensor = WeightTensorRDGs[Idx];
-		FRDGBufferRef Buffer = RDGBuilder.RegisterExternalBuffer(PooledBuffer);
-		Tensor.SetBuffer(Buffer);
+		// Only register external buffer when weight tensor is NOT zero-sized.
+		if(Tensor.HasPreparedData())
+		{
+			FRDGBufferRef Buffer = RDGBuilder.RegisterExternalBuffer(PooledBuffer);
+			Tensor.SetBuffer(Buffer);
+		}
+		else
+		{
+			Tensor.SetBuffer(nullptr);
+		}
 	}
 
 	return true;
@@ -195,7 +203,7 @@ void FModelInstance::AddDispatchOps_RenderThread(FRDGBuilder& GraphBuilder)
 		OutputTensors.Reset(OperatorOutputTensorIndices.Num());
 		for (int32 i : OperatorOutputTensorIndices[Idx])
 		{
-			AllOutputTensorConstant &= AllTensorRDGRefs[i]->HasPreparedData();
+			AllOutputTensorConstant &= AllTensorRDGRefs[i]->IsConstant();
 			OutputTensors.Add(AllTensorRDGRefs[i]);
 		}
 
@@ -206,7 +214,7 @@ void FModelInstance::AddDispatchOps_RenderThread(FRDGBuilder& GraphBuilder)
 		}
 	}
 
-	//If a model output is constant we upload to it (a user provided GPU buffer).
+	//If a model output is constant (and non-empty) we upload to it (a user provided GPU buffer).
 	for (const FTensorRDG& OutputTensor : OutputTensorRDGs)
 	{
 		if (OutputTensor.HasPreparedData())
@@ -301,15 +309,15 @@ void EnqueueTensorUpload(TArray<TRefCountPtr<FRDGPooledBuffer>>& OutExternalRDGR
 		OutExternalRDGResources.Reset();
 		OutExternalRDGResources.SetNum(TensorToUploadRDGs.Num());
 
-		bool AtLeastOneConstantTensor = false;
+		bool AtLeastOneConstantNonEmptyTensor = false;
 		for (const FTensorRDG& TensorRDG : TensorToUploadRDGs)
 		{
 			if (TensorRDG.HasPreparedData())
 			{
-				AtLeastOneConstantTensor = true;
+				AtLeastOneConstantNonEmptyTensor = true;
 			}
 		}
-		if (!AtLeastOneConstantTensor)
+		if (!AtLeastOneConstantNonEmptyTensor)
 		{
 			return;
 		}
diff --git a/Engine/Plugins/Experimental/NNERuntimeRDG/Source/NNERuntimeRDGUtils/Internal/NNERuntimeRDGUtilsModelBuilder.h b/Engine/Plugins/Experimental/NNERuntimeRDG/Source/NNERuntimeRDGUtils/Internal/NNERuntimeRDGUtilsModelBuilder.h
--- a/Engine/Plugins/Experimental/NNERuntimeRDG/Source/NNERuntimeRDGUtils/Internal/NNERuntimeRDGUtilsModelBuilder.h
+++ b/Engine/Plugins/Experimental/NNERuntimeRDG/Source/NNERuntimeRDGUtils/Internal/NNERuntimeRDGUtilsModelBuilder.h
@@ -50,7 +50,8 @@ public:
 	// Serialize the model to given array
 	virtual bool End(TArray<uint8>& Data) = 0;
 
-	virtual FHTensor AddTensor(const FString& Name, ENNETensorDataType DataType, TArrayView<const int32> Shape, const void* Data = nullptr, uint64 DataSize = 0) = 0;
+	virtual FHTensor AddTensor(const FString& Name, ENNETensorDataType DataType, TArrayView<const int32> Shape) = 0;
+	virtual FHTensor AddConstantTensor(const FString& Name, ENNETensorDataType DataType, TArrayView<const int32> Shape, const void* Data, uint64 DataSize) = 0;
 	virtual bool AddInput(FHTensor InTensor) = 0;
 	virtual bool AddOutput(FHTensor OutTensor) = 0;
 	virtual FHOperator AddOperator(const FString& Type, const FString& Domain, TOptional<uint32> Version = FNullOpt{0}, const FString& Name = TEXT("")) = 0;
diff --git a/Engine/Plugins/Experimental/NNERuntimeRDG/Source/NNERuntimeRDGUtils/Private/NNERuntimeRDGUtilsModelBuilderNNE.cpp b/Engine/Plugins/Experimental/NNERuntimeRDG/Source/NNERuntimeRDGUtils/Private/NNERuntimeRDGUtilsModelBuilderNNE.cpp
--- a/Engine/Plugins/Experimental/NNERuntimeRDG/Source/NNERuntimeRDGUtils/Private/NNERuntimeRDGUtilsModelBuilderNNE.cpp
+++ b/Engine/Plugins/Experimental/NNERuntimeRDG/Source/NNERuntimeRDGUtils/Private/NNERuntimeRDGUtilsModelBuilderNNE.cpp
@@ -51,17 +51,24 @@ public:
 		return !Data.IsEmpty();
 	}
 
-	virtual FHTensor AddTensor(const FString& Name, ENNETensorDataType DataType, TArrayView<const int32> Shape, const void* Data, uint64 DataSize)
+	virtual FHTensor AddTensor(const FString& Name, ENNETensorDataType DataType, TArrayView<const int32> Shape) override
 	{
 		TArray<int32, TInlineAllocator<NNE::FTensorShape::MaxRank>> NNEShape;
 		for (int32 i = 0; i < Shape.Num(); ++i)
 		{
 			//Allow caller to use 0 for variable dimensions for inputs/outputs, NNE use -1.
 			//RDG not supporting 0 sized dimension at the moment.
-			NNEShape.Emplace(!Data && Shape[i] == 0 ? -1 : Shape[i]);
+			NNEShape.Emplace(Shape[i] == 0 ? -1 : Shape[i]);
 		}
 
-		int32 Idx = AddTensor(Name, NNEShape, DataType, Data, DataSize);
+		int32 Idx = AddTensor(Name, NNEShape, DataType);
+
+		return MakeHandle<EHandleType::Tensor>(reinterpret_cast<void*>((int64)Idx));
+	}
+
+	virtual FHTensor AddConstantTensor(const FString& Name, ENNETensorDataType DataType, TArrayView<const int32> Shape, const void* Data, uint64 DataSize) override
+	{
+		int32 Idx = AddTensor(Name, Shape, DataType, Data, DataSize);
 
 		return MakeHandle<EHandleType::Tensor>(reinterpret_cast<void*>((int64)Idx));
 	}
@@ -178,7 +185,7 @@ public:
 
 private:
 
-	int32 AddTensor(const FString& InName, TArrayView<const int32> InShape, ENNETensorDataType InDataType, const void* Data, uint64 DataSize)
+	int32 AddTensor(const FString& InName, TArrayView<const int32> InShape, ENNETensorDataType InDataType, const void* Data = nullptr, uint64 DataSize = 0u)
 	{
 		int32 Idx = -1;
 
@@ -200,19 +207,19 @@ private:
 			if (Data)
 			{
 				Desc.Type = ENNEFormatTensorType::Initializer;
+				Desc.DataSize = DataSize;
 
 				// Handle empty data initializers, i.e. when DataSize is 0
 				if (DataSize)
 				{
 					Desc.DataOffset = Format.TensorData.AddUninitialized(DataSize);
-					Desc.DataSize = DataSize;
 
 					FMemory::Memcpy(Format.TensorData.GetData() + Desc.DataOffset, Data, DataSize);
 				}
 			}
 			else
 			{
-				if (Desc.DataType == ENNETensorDataType::None && InName.IsEmpty())
+				if (InName.IsEmpty())
 				{
 					Desc.Type = ENNEFormatTensorType::Empty;
 				}
diff --git a/Engine/Plugins/Experimental/NNERuntimeRDG/Source/NNERuntimeRDGUtils/Private/NNERuntimeRDGUtilsModelBuilderONNX.cpp b/Engine/Plugins/Experimental/NNERuntimeRDG/Source/NNERuntimeRDGUtils/Private/NNERuntimeRDGUtilsModelBuilderONNX.cpp
--- a/Engine/Plugins/Experimental/NNERuntimeRDG/Source/NNERuntimeRDGUtils/Private/NNERuntimeRDGUtilsModelBuilderONNX.cpp
+++ b/Engine/Plugins/Experimental/NNERuntimeRDG/Source/NNERuntimeRDGUtils/Private/NNERuntimeRDGUtilsModelBuilderONNX.cpp
@@ -112,32 +112,41 @@ public:
 		return res;
 	}
 
-	virtual FHTensor AddTensor(const FString& Name, ENNETensorDataType DataType, TArrayView<const int32> Shape, const void* Data, uint64 DataSize) override
+	virtual FHTensor AddTensor(const FString& Name, ENNETensorDataType DataType, TArrayView<const int32> Shape) override
 	{
-		onnx::ValueInfoProto* Value = onnx::ValueInfoProto().New(Graph->GetArena());
+		onnx::ValueInfoProto* Value = AddValueInfoProto(Name, DataType, Shape);
 		
-		if (!SetValue(Value, Name, DataType, Shape))
+		if(Value == nullptr)
 		{
 			return FHTensor();
 		}
 
-		if (Data)
+		return MakeHandle<EHandleType::Tensor>(Value);
+	}
+
+	virtual FHTensor AddConstantTensor(const FString& Name, ENNETensorDataType DataType, TArrayView<const int32> Shape, const void* Data, uint64 DataSize) override
+	{
+		onnx::ValueInfoProto* Value = AddValueInfoProto(Name, DataType, Shape);
+		
+		if(Value == nullptr)
 		{
-			onnx::TensorProto* Tensor = Graph->mutable_initializer()->Add();
+			return FHTensor();
+		}
 
-			Tensor->set_name(TCHAR_TO_ANSI(*Name));
-			Tensor->set_data_type(ToTensorProtoDataType(DataType));
-			for (int32 Idx = 0; Idx < Shape.Num(); ++Idx)
-			{
-				Tensor->add_dims(Shape[Idx]);
-			}
+		onnx::TensorProto* Tensor = Graph->mutable_initializer()->Add();
 
-			checkf(DataType != ENNETensorDataType::Char, TEXT("Char tensors not supported yet!"));
-			// raw_data is not supported for strings, when implementing those one will need to use string_data from TensorProto
-			std::string* raw_data = Tensor->mutable_raw_data();
-			raw_data->assign((const char*)Data, DataSize);
+		Tensor->set_name(TCHAR_TO_ANSI(*Name));
+		Tensor->set_data_type(ToTensorProtoDataType(DataType));
+		for (int32 Idx = 0; Idx < Shape.Num(); ++Idx)
+		{
+			Tensor->add_dims(Shape[Idx]);
 		}
 
+		checkf(DataType != ENNETensorDataType::Char, TEXT("Char tensors not supported yet!"));
+		// raw_data is not supported for strings, when implementing those one will need to use string_data from TensorProto
+		std::string* raw_data = Tensor->mutable_raw_data();
+		raw_data->assign((const char*)Data, DataSize);
+
 		return MakeHandle<EHandleType::Tensor>(Value);
 	}
 
@@ -268,6 +277,18 @@ public:
 
 private:
 
+	onnx::ValueInfoProto* AddValueInfoProto(const FString& Name, ENNETensorDataType DataType, TArrayView<const int32> Shape)
+	{
+		onnx::ValueInfoProto* Value = onnx::ValueInfoProto().New(Graph->GetArena());
+		
+		if (!SetValue(Value, Name, DataType, Shape))
+		{
+			return nullptr;
+		}
+
+		return Value;
+	}
+
 	bool SetValue(onnx::ValueInfoProto* Value, const FString& Name, ENNETensorDataType DataType, const TArrayView<const int32>& InShape)
 	{
 		onnx::TypeProto*		Type = Value->mutable_type();
@@ -361,7 +382,7 @@ bool CreateONNXModelForOperator(const FString& OperatorName, int32 IrVersion, in
 		const TConstArrayView<uint8>& Data = InWeightTensorsData[Idx];
 		check(Data.Num() == Desc.GetDataSize());
 		ModelBuilderONNXHelper::BuildShapeForModel(false, Desc.GetShape(), ShapeForModel);
-		IModelBuilder::FHTensor Tensor = Builder->AddTensor(Desc.GetName(), Desc.GetDataType(), ShapeForModel, Data.GetData(), Data.Num());
+		IModelBuilder::FHTensor Tensor = Builder->AddConstantTensor(Desc.GetName(), Desc.GetDataType(), ShapeForModel, Data.GetData(), Data.Num());
 
 		WeightTensors.Emplace(Tensor);
 	}
diff --git a/Engine/Plugins/Experimental/NNERuntimeRDG/Source/NNERuntimeRDGUtils/Private/NNERuntimeRDGUtilsModelOptimizerNNE.cpp b/Engine/Plugins/Experimental/NNERuntimeRDG/Source/NNERuntimeRDGUtils/Private/NNERuntimeRDGUtilsModelOptimizerNNE.cpp
--- a/Engine/Plugins/Experimental/NNERuntimeRDG/Source/NNERuntimeRDGUtils/Private/NNERuntimeRDGUtilsModelOptimizerNNE.cpp
+++ b/Engine/Plugins/Experimental/NNERuntimeRDG/Source/NNERuntimeRDGUtils/Private/NNERuntimeRDGUtilsModelOptimizerNNE.cpp
@@ -259,7 +259,7 @@ namespace ModelOptimizerNNEHelper
 				}
 
 				IModelBuilder::FHTensor TensorInitializer =
-					Builder->AddTensor(FString(ANSI_TO_TCHAR(Output.name().c_str())) + TEXT("_NNEInitializer"), DataType, Shape, Data, DataSize);
+					Builder->AddConstantTensor(FString(ANSI_TO_TCHAR(Output.name().c_str())) + TEXT("_NNEInitializer"), DataType, Shape, Data, DataSize);
 
 				const FString IdentityOpType = TEXT("Identity");
 				TOptional<uint32> OpVersion = GetOpVersionFromOpsetVersion(IdentityOpType, (int) ModelProto.opset_import(0).version());
@@ -358,6 +358,8 @@ namespace ModelOptimizerNNEHelper
 				const void* Data = nullptr;
 				uint64 DataSize = 0;
 
+				IModelBuilder::FHTensor Tensor;
+
 				const onnx::TensorProto* Initializer = GetInitializerFromGraphProto(Graph, TensorName);
 				if (Initializer)
 				{
@@ -366,6 +368,7 @@ namespace ModelOptimizerNNEHelper
 						UE_LOG(LogNNE, Error, TEXT("Tensor data could not be loaded for weight '%s' in node '%s' of type '%s'"), ANSI_TO_TCHAR(TensorName.c_str()), *NNEOpName, *NNEOpType);
 						return false;
 					}
+					Tensor = Builder->AddConstantTensor(ANSI_TO_TCHAR(TensorName.c_str()), DataType, Shape, Data, DataSize);
 				}
 				else if (!TensorName.empty())
 				{
@@ -377,9 +380,12 @@ namespace ModelOptimizerNNEHelper
 					}
 
 					GetTensorInfoFromONNXValueInfo(*ValueInfoProto, Shape, DataType);
+					Tensor = Builder->AddTensor(ANSI_TO_TCHAR(TensorName.c_str()), DataType, Shape);
+				}
+				else // Empty tensor
+				{
+					Tensor = Builder->AddConstantTensor(/* Name */ TEXT(""), ENNETensorDataType::None, /* Shape */ {0}, /* Data */ nullptr, /* DataSize */ 0);
 				}
-
-				IModelBuilder::FHTensor Tensor = Builder->AddTensor(ANSI_TO_TCHAR(TensorName.c_str()), DataType, Shape, Data, DataSize);
 
 				Builder->AddOperatorInput(Op, Tensor);
 			}
