NNE - Fix Hlsl Reduce bug when multiple axes are provided

Problems:
* RDG Hlsl Reduce (all) gives wrong results when multiple axes to reduce on are provided.

Solutions:
* Fix a small host code indexing bug that causes the wrong results.

Tests:
* Win64 Standalone tests: all
* Win64 Editor tests: all

#rnx

#virtualized

diff --git a/Engine/Plugins/Experimental/NNERuntimeRDG/Source/NNERuntimeRDG/Private/Hlsl/NNERuntimeRDGReduce.cpp b/Engine/Plugins/Experimental/NNERuntimeRDG/Source/NNERuntimeRDG/Private/Hlsl/NNERuntimeRDGReduce.cpp
--- a/Engine/Plugins/Experimental/NNERuntimeRDG/Source/NNERuntimeRDG/Private/Hlsl/NNERuntimeRDGReduce.cpp
+++ b/Engine/Plugins/Experimental/NNERuntimeRDG/Source/NNERuntimeRDG/Private/Hlsl/NNERuntimeRDGReduce.cpp
@@ -6,6 +6,8 @@
 #include "Helper/NNERuntimeRDGOperatorHelper.h"
 #include "NNETensor.h"
 #include "NNETypes.h"
+#include "Algo/Sort.h"
+#include "Misc/EnumerateRange.h"
 #include "RenderGraphUtils.h"
 
 namespace UE::NNERuntimeRDG::Private::Hlsl
@@ -31,7 +33,7 @@ namespace UE::NNERuntimeRDG::Private::Hlsl
 
 	private:
 
-		TArray<int32, TInlineAllocator<NNE::FTensorShape::MaxRank>> Axes;
+		TArray<int32, TInlineAllocator<NNE::FTensorShape::MaxRank>> Axes; // Must be sorted in descending order
 		int32 KeepDims = 1;
 
 	public:
@@ -74,6 +76,8 @@ namespace UE::NNERuntimeRDG::Private::Hlsl
 					Axis += InputRank;
 				}
 			}
+
+			Algo::Sort(Axes, TGreater<>());
 			
 			TArray<uint32> OutputShape;
 
@@ -121,7 +125,6 @@ namespace UE::NNERuntimeRDG::Private::Hlsl
 			}
 			
 			Axes = Attributes.GetValueOrDefault<TArray<int32>>(TEXT("axes"), AxesDefault);
-			Axes.Sort();
 
 			if (!bAxesAsInput && Axes.Num() == 0)
 			{
@@ -158,20 +161,22 @@ namespace UE::NNERuntimeRDG::Private::Hlsl
 			{
 				FRDGBufferRef CurrInput = Input.GetBuffer();
 				TArray<uint32, TInlineAllocator<NNE::FTensorShape::MaxRank>> CurrInputShape(Input.GetShape().GetData());
-				for (int32 i = Axes.Num()-1; i >= 0; --i)
+
+				// Iterate axes in descending order
+				for (TConstEnumerateRef<int32> AxisRef : EnumerateRange(Axes))
 				{
-					const int32 Axis = Axes[i];
+					const int32 Axis = *AxisRef;
 
 					TReduceCS::FParameters* Parameters = GraphBuilder.AllocParameters<TReduceCS::FParameters>();
 					TReduceCS::FillInParameters(CurrInputShape, Axis, Parameters);
 
 					FRDGBufferRef CurrOutput = nullptr;
-					if (i != 0)
+					if (AxisRef.GetIndex() < Axes.Num() - 1)
 					{
 						const FRDGBufferDesc TempBufferDesc = FRDGBufferDesc::CreateBufferDesc(Output.GetElementByteSize(), Parameters->NumElemBeforeAxis * Parameters->NumElemAfterAxis);
 						CurrOutput = GraphBuilder.CreateBuffer(TempBufferDesc, TEXT("NNE.Operator.Hlsl.Reduce.TempBuffer"), ERDGBufferFlags::None);
 					}
-					else
+					else // Last iteration
 					{
 						CurrOutput = Output.GetBuffer();
 					}
@@ -179,7 +184,7 @@ namespace UE::NNERuntimeRDG::Private::Hlsl
 
 					TReduceCS::EnqueueRDG(GraphBuilder, Parameters, CurrInput, CurrOutput, ReduceOperatorType);
 					CurrInput = CurrOutput;
-					CurrInputShape[i] = 1;
+					CurrInputShape[Axis] = 1;
 				}
 			}
 		}
