NNE - Settle empty tensors in NNE format

Problems:
* RDG Hlsl breaks when processing an empty tensor in `FModelInstance::Init()` where `AllSymbolicTensorDescs` map is accessed at the empty tensor's index.. The bug comes from not treating empty tensors as first class citizens in NNE format.
* NNEQA Json description of tests doesn't allow to specify empty tensors.
* Hlsl OPs like Resize cannot be fully tested until empty tensors are correctly working.

Solutions:
* Due to the unavailability of FNNEFormatTensorDesc in Hlsl OP classes or validate functions, one cannot check if Type (ENNEFormatTensorType) is Empty directly. Therefore from now on we characterize empty tensors in NNE format as having both ENNEFormatTensorType set to Empty and ENNETensorDataType to None, with ENNETensorDataType::None being reserved for empty tensors (it's also used as a temporary value by NNE model builder while constructing the graph). Hlsl OP classes and validate functions have a way to access the data type.
  The description of an empty tensor in NNE format is therefore settled to: non-empty name (autogenerated), None data type, Empty tensor type, Shape = [0].
* Fix RDG handling of empty tensors by treating them as first class citizens.
* Enhance NNEQA Json format of a tensor with a way to specify empty tensors: "empty" type.
* Refactor model builders (ONNX/NNE) to consider empty tensors as a special case. Also make the necessary changes to allow Json-testing of empty tensors.
* Enable Hlsl Resize's Sizes tensor. Requires to have Scales tensor empty.

Tests:
* Test that empty tensors are working by new Resize tests soliciting the Sizes tensor.
* Win64 Standalone tests: all
* Win64 Editor tests: all

#rnx

#virtualized

diff --git a/Engine/Plugins/Experimental/NNERuntimeRDG/Source/NNERuntimeRDG/Private/Hlsl/NNERuntimeRDGResize.cpp b/Engine/Plugins/Experimental/NNERuntimeRDG/Source/NNERuntimeRDG/Private/Hlsl/NNERuntimeRDGResize.cpp
--- a/Engine/Plugins/Experimental/NNERuntimeRDG/Source/NNERuntimeRDG/Private/Hlsl/NNERuntimeRDGResize.cpp
+++ b/Engine/Plugins/Experimental/NNERuntimeRDG/Source/NNERuntimeRDG/Private/Hlsl/NNERuntimeRDGResize.cpp
@@ -71,7 +71,7 @@ namespace UE::NNERuntimeRDG::Private::Hlsl
 				return -1;
 			}
 
-			if(InputTensors.Num() == 4 && InputTensors[2]->GetVolume() != 0)
+			if(InputTensors.Num() == 4 && InputTensors[2]->GetDataType() != ENNETensorDataType::None)
 			{
 				UE_LOG(LogNNE, Warning, TEXT("Hlsl Resize: both `sizes` and `scales` were provided."));
 				return -1;
@@ -450,14 +450,6 @@ namespace UE::NNERuntimeRDG::Private::Hlsl
 			}
 		}
 
-		// NOTE: the current RDG Hlsl infrastructure does NOT support empty tensors (FModelInstance::Init() breaks), 
-		// therefore it's not possible to specify the Sizes tensor (as Scales would need to be empty).
-		if(InputTypes.Num() == 4)
-		{
-			UE_LOG(LogNNE, Warning, TEXT("Sizes tensor is not currently supported, please use Scales tensor."));
-			return false;
-		}
-
 		return bIsValid;
 	}
 
diff --git a/Engine/Plugins/Experimental/NNERuntimeRDG/Source/NNERuntimeRDG/Private/NNERuntimeRDGModel.cpp b/Engine/Plugins/Experimental/NNERuntimeRDG/Source/NNERuntimeRDG/Private/NNERuntimeRDGModel.cpp
--- a/Engine/Plugins/Experimental/NNERuntimeRDG/Source/NNERuntimeRDG/Private/NNERuntimeRDGModel.cpp
+++ b/Engine/Plugins/Experimental/NNERuntimeRDG/Source/NNERuntimeRDG/Private/NNERuntimeRDGModel.cpp
@@ -31,6 +31,7 @@ bool FModelInstanceRDG::LoadModel(TConstArrayView<uint8> ModelData, FNNERuntimeF
 	WeightTensorIndices.Empty();
 	InputTensorIndices.Empty();
 	OutputTensorIndices.Empty();
+	EmptyTensorIndices.Empty();
 	OperatorInputTensorIndices.Empty();
 	OperatorOutputTensorIndices.Empty();
 
@@ -44,12 +45,19 @@ bool FModelInstanceRDG::LoadModel(TConstArrayView<uint8> ModelData, FNNERuntimeF
 		const NNE::FSymbolicTensorShape SymbolicShape = NNE::FSymbolicTensorShape::Make(FormatTensorDesc.Shape);
 		const NNE::FTensorDesc SymbolicTensor = NNE::FTensorDesc::Make(FormatTensorDesc.Name, SymbolicShape, FormatTensorDesc.DataType);
 
-		if (Format.Tensors[Idx].Type != ENNEFormatTensorType::Empty)
+		if (FormatTensorDesc.Type != ENNEFormatTensorType::Empty && FormatTensorDesc.DataType == ENNETensorDataType::None)
 		{
-			AllSymbolicTensorDescs.Emplace(Idx, SymbolicTensor);
+			UE_LOG(LogNNE, Error, TEXT("Tensor %s has invalid format: Data type None is reserved for empty tensors."), *SymbolicTensor.GetName());
+			return false;
 		}
 
-		if (FormatTensorDesc.Type == ENNEFormatTensorType::Input)
+		AllSymbolicTensorDescs.Emplace(Idx, SymbolicTensor);
+
+		if (FormatTensorDesc.Type == ENNEFormatTensorType::Empty)
+		{
+			EmptyTensorIndices.Emplace(Idx);
+		}
+		else if (FormatTensorDesc.Type == ENNEFormatTensorType::Input)
 		{
 			InputTensorIndices.Emplace(Idx);
 			InputSymbolicTensors.Emplace(SymbolicTensor);
@@ -154,6 +162,17 @@ FModelInstanceRDG::ESetInputTensorShapesStatus FModelInstanceRDG::SetInputTensor
 		AllTensorRDGRefs.Emplace(Idx, &OutputTensorRDGs[i]);
 	}
 
+	EmptyTensorRDGs.Reset(EmptyTensorIndices.Num());
+	for (int32 i = 0; i < EmptyTensorIndices.Num(); ++i)
+	{
+		const int32 Idx = EmptyTensorIndices[i];
+		const NNE::FTensorDesc& TensorDesc = AllSymbolicTensorDescs[Idx];
+		const NNE::FTensorShape TensorShape = NNE::FTensorShape::MakeFromSymbolic(TensorDesc.GetShape());
+
+		EmptyTensorRDGs.Emplace(FTensorRDG::Make(TensorDesc, TensorShape, nullptr));
+		AllTensorRDGRefs.Emplace(Idx, &EmptyTensorRDGs[i]);
+	}
+
 	checkCode(
 		checkf(AllTensorRDGRefs.Num() == AllSymbolicTensorDescs.Num(), TEXT("Some tensor was not allocated for model preparation."));
 	);
@@ -177,7 +196,7 @@ FModelInstanceRDG::ESetInputTensorShapesStatus FModelInstanceRDG::SetInputTensor
 		OutputTensorShapes.Emplace(AllTensorRDGRefs[OutputIndices]->GetShape());
 	}
 
-	check(InputTensorIndices.Num() + OutputTensorIndices.Num() + WeightTensorIndices.Num() + IntermediateTensorIndices.Num() == AllTensorRDGRefs.Num());
+	check(InputTensorIndices.Num() + OutputTensorIndices.Num() + WeightTensorIndices.Num() + IntermediateTensorIndices.Num() + EmptyTensorIndices.Num() == AllTensorRDGRefs.Num());
 	check(InputTensorShapes.Num() == InputSymbolicTensors.Num());
 	check(OutputTensorShapes.Num() == OutputSymbolicTensors.Num());
 	check(WeightTensorIndices.Num() == WeightTensorRDGs.Num());
diff --git a/Engine/Plugins/Experimental/NNERuntimeRDG/Source/NNERuntimeRDG/Private/NNERuntimeRDGModel.h b/Engine/Plugins/Experimental/NNERuntimeRDG/Source/NNERuntimeRDG/Private/NNERuntimeRDGModel.h
--- a/Engine/Plugins/Experimental/NNERuntimeRDG/Source/NNERuntimeRDG/Private/NNERuntimeRDGModel.h
+++ b/Engine/Plugins/Experimental/NNERuntimeRDG/Source/NNERuntimeRDG/Private/NNERuntimeRDGModel.h
@@ -41,6 +41,7 @@ namespace UE::NNERuntimeRDG::Private
 		TArray<int32>				WeightTensorIndices;
 		TArray<int32>				InputTensorIndices;
 		TArray<int32>				OutputTensorIndices;
+		TArray<int32>				EmptyTensorIndices;
 
 		//Tensor indices by operator
 		TArray<TArray<uint32>>		OperatorInputTensorIndices;
@@ -50,6 +51,7 @@ namespace UE::NNERuntimeRDG::Private
 		FTensorRDGRefMap			AllTensorRDGRefs;
 		FTensorRDGArray				InputTensorRDGs;
 		FTensorRDGArray				OutputTensorRDGs;
+		FTensorRDGArray				EmptyTensorRDGs;
 		FTensorRDGArray				IntermediateTensorRDGs;
 		FTensorRDGArray				WeightTensorRDGs;
 	};
diff --git a/Engine/Plugins/Experimental/NNERuntimeRDG/Source/NNERuntimeRDG/Private/NNERuntimeRDGModelHlsl.cpp b/Engine/Plugins/Experimental/NNERuntimeRDG/Source/NNERuntimeRDG/Private/NNERuntimeRDGModelHlsl.cpp
--- a/Engine/Plugins/Experimental/NNERuntimeRDG/Source/NNERuntimeRDG/Private/NNERuntimeRDGModelHlsl.cpp
+++ b/Engine/Plugins/Experimental/NNERuntimeRDG/Source/NNERuntimeRDG/Private/NNERuntimeRDGModelHlsl.cpp
@@ -251,6 +251,10 @@ int FModelInstance::PrepareTensorShapesAndData()
 		{
 			AllInitializedTensors[Idx] = true;
 		}
+		for (int32 Idx : EmptyTensorIndices)
+		{
+			AllInitializedTensors[Idx] = true;
+		}
 	);
 
 	//Release uploaded GPU side constants tensors.
diff --git a/Engine/Plugins/Experimental/NNERuntimeRDG/Source/NNERuntimeRDGUtils/Internal/NNERuntimeRDGUtilsModelBuilder.h b/Engine/Plugins/Experimental/NNERuntimeRDG/Source/NNERuntimeRDGUtils/Internal/NNERuntimeRDGUtilsModelBuilder.h
--- a/Engine/Plugins/Experimental/NNERuntimeRDG/Source/NNERuntimeRDGUtils/Internal/NNERuntimeRDGUtilsModelBuilder.h
+++ b/Engine/Plugins/Experimental/NNERuntimeRDG/Source/NNERuntimeRDGUtils/Internal/NNERuntimeRDGUtilsModelBuilder.h
@@ -52,6 +52,7 @@ public:
 
 	virtual FHTensor AddTensor(const FString& Name, ENNETensorDataType DataType, TArrayView<const int32> Shape) = 0;
 	virtual FHTensor AddConstantTensor(const FString& Name, ENNETensorDataType DataType, TArrayView<const int32> Shape, const void* Data, uint64 DataSize) = 0;
+	virtual FHTensor AddEmptyTensor() = 0;
 	virtual bool AddInput(FHTensor InTensor) = 0;
 	virtual bool AddOutput(FHTensor OutTensor) = 0;
 	virtual FHOperator AddOperator(const FString& Type, const FString& Domain, TOptional<uint32> Version = FNullOpt{0}, const FString& Name = TEXT("")) = 0;
diff --git a/Engine/Plugins/Experimental/NNERuntimeRDG/Source/NNERuntimeRDGUtils/Private/NNERuntimeRDGUtilsModelBuilderNNE.cpp b/Engine/Plugins/Experimental/NNERuntimeRDG/Source/NNERuntimeRDGUtils/Private/NNERuntimeRDGUtilsModelBuilderNNE.cpp
--- a/Engine/Plugins/Experimental/NNERuntimeRDG/Source/NNERuntimeRDGUtils/Private/NNERuntimeRDGUtilsModelBuilderNNE.cpp
+++ b/Engine/Plugins/Experimental/NNERuntimeRDG/Source/NNERuntimeRDGUtils/Private/NNERuntimeRDGUtilsModelBuilderNNE.cpp
@@ -61,14 +61,23 @@ public:
 			NNEShape.Emplace(Shape[i] == 0 ? -1 : Shape[i]);
 		}
 
-		int32 Idx = AddTensor(Name, NNEShape, DataType);
+		int32 Idx = AddTensor(GetBaseTensorDesc(Name, NNEShape, DataType));
 
 		return MakeHandle<EHandleType::Tensor>(reinterpret_cast<void*>((int64)Idx));
 	}
 
 	virtual FHTensor AddConstantTensor(const FString& Name, ENNETensorDataType DataType, TArrayView<const int32> Shape, const void* Data, uint64 DataSize) override
 	{
-		int32 Idx = AddTensor(Name, Shape, DataType, Data, DataSize);
+		FNNEFormatTensorDesc Desc = GetBaseTensorDesc(Name, Shape, DataType);
+		AddInitializer(Desc, Data, DataSize);
+		int32 Idx = AddTensor(MoveTemp(Desc));
+
+		return MakeHandle<EHandleType::Tensor>(reinterpret_cast<void*>((int64)Idx));
+	}
+
+	virtual FHTensor AddEmptyTensor() override
+	{
+		int32 Idx = AddTensor(GetEmptyTensorDesc());
 
 		return MakeHandle<EHandleType::Tensor>(reinterpret_cast<void*>((int64)Idx));
 	}
@@ -185,11 +194,53 @@ public:
 
 private:
 
-	int32 AddTensor(const FString& InName, TArrayView<const int32> InShape, ENNETensorDataType InDataType, const void* Data = nullptr, uint64 DataSize = 0u)
+	static FNNEFormatTensorDesc GetBaseTensorDesc(const FString& InName, TArrayView<const int32> InShape, ENNETensorDataType InDataType)
+	{
+		FNNEFormatTensorDesc&&	Desc{};
+
+		Desc.Name = InName;
+		Desc.Shape = InShape;
+		Desc.Type = ENNEFormatTensorType::None;
+		Desc.DataType = InDataType;
+
+		return Desc;
+	}
+
+	void AddInitializer(FNNEFormatTensorDesc& Desc, const void* Data, uint64 DataSize = 0u)
+	{
+		Desc.Type = ENNEFormatTensorType::Initializer;
+		Desc.DataSize = DataSize;
+
+		// Handle empty data initializers, i.e. when DataSize is 0
+		if (DataSize)
+		{
+			Desc.DataOffset = Format.TensorData.AddUninitialized(DataSize);
+			FMemory::Memcpy(Format.TensorData.GetData() + Desc.DataOffset, Data, DataSize);
+		}
+	}
+
+	FString GenerateEmptyTensorName()
+	{
+		return FString::Printf(TEXT("__NNE_EmptyTensor_%u"), EmptyTensorCounter++);
+	}
+
+	FNNEFormatTensorDesc GetEmptyTensorDesc()
+	{
+		FNNEFormatTensorDesc&&	Desc{};
+
+		Desc.Name = GenerateEmptyTensorName();
+		Desc.Shape = { 0 };
+		Desc.Type = ENNEFormatTensorType::Empty;
+		Desc.DataType = ENNETensorDataType::None;
+
+		return Desc;
+	}
+
+	int32 AddTensor(FNNEFormatTensorDesc&& InTensorDesc)
 	{
 		int32 Idx = -1;
 
-		int32* Val = TensorMap.Find(InName);
+		int32* Val = TensorMap.Find(InTensorDesc.Name);
 
 		if (Val)
 		{
@@ -197,38 +248,10 @@ private:
 		}
 		else
 		{
-			FNNEFormatTensorDesc	Desc{};
-
-			Desc.Name = InName;
-			Desc.Shape = InShape;
-			Desc.Type = ENNEFormatTensorType::None;
-			Desc.DataType = InDataType;
-
-			if (Data)
-			{
-				Desc.Type = ENNEFormatTensorType::Initializer;
-				Desc.DataSize = DataSize;
-
-				// Handle empty data initializers, i.e. when DataSize is 0
-				if (DataSize)
-				{
-					Desc.DataOffset = Format.TensorData.AddUninitialized(DataSize);
-
-					FMemory::Memcpy(Format.TensorData.GetData() + Desc.DataOffset, Data, DataSize);
-				}
-			}
-			else
-			{
-				if (InName.IsEmpty())
-				{
-					Desc.Type = ENNEFormatTensorType::Empty;
-				}
-			}
-
-			Format.Tensors.Add(Desc);
+			Format.Tensors.Add(MoveTemp(InTensorDesc));
 			Idx = Format.Tensors.Num() - 1;
 
-			TensorMap.Add(InName, Idx);
+			TensorMap.Add(Format.Tensors[Idx].Name, Idx);
 		}
 
 		return Idx;
@@ -237,6 +260,7 @@ private:
 
 	FNNERuntimeFormat	Format;
 	TMap<FString, int32>	TensorMap;
+	uint32 EmptyTensorCounter = 0;
 };
 
 TUniquePtr<IModelBuilder> CreateNNEModelBuilder()
diff --git a/Engine/Plugins/Experimental/NNERuntimeRDG/Source/NNERuntimeRDGUtils/Private/NNERuntimeRDGUtilsModelBuilderONNX.cpp b/Engine/Plugins/Experimental/NNERuntimeRDG/Source/NNERuntimeRDGUtils/Private/NNERuntimeRDGUtilsModelBuilderONNX.cpp
--- a/Engine/Plugins/Experimental/NNERuntimeRDG/Source/NNERuntimeRDGUtils/Private/NNERuntimeRDGUtilsModelBuilderONNX.cpp
+++ b/Engine/Plugins/Experimental/NNERuntimeRDG/Source/NNERuntimeRDGUtils/Private/NNERuntimeRDGUtilsModelBuilderONNX.cpp
@@ -150,6 +150,12 @@ public:
 		return MakeHandle<EHandleType::Tensor>(Value);
 	}
 
+	virtual FHTensor AddEmptyTensor() override
+	{
+		// ONNX empty tensors are indicated by an empty name
+		return AddTensor(/* Name */ TEXT(""), ENNETensorDataType::None, /* Shape */ {0});
+	}
+
 	virtual bool AddInput(FHTensor Handle) override
 	{
 		onnx::ValueInfoProto* Value = ModelBuilderONNXHelper::OnnxTensorCast(Handle);
@@ -355,7 +361,9 @@ bool CreateONNXModelForOperator(const FString& OperatorName, int32 IrVersion, in
 	{
 		const NNE::Internal::FTensor& Desc = InInputTensors[Idx];
 		ModelBuilderONNXHelper::BuildShapeForModel(bUseVariadicShapeForModel, Desc.GetShape(), ShapeForModel);
-		IModelBuilder::FHTensor Tensor = Builder->AddTensor(Desc.GetName(), Desc.GetDataType(), ShapeForModel);
+		IModelBuilder::FHTensor Tensor = Desc.GetDataType() == ENNETensorDataType::None ?
+			  Builder->AddEmptyTensor()
+			: Builder->AddTensor(Desc.GetName(), Desc.GetDataType(), ShapeForModel);
 
 		InputTensors.Emplace(Tensor);
 		Builder->AddInput(Tensor);
@@ -382,7 +390,9 @@ bool CreateONNXModelForOperator(const FString& OperatorName, int32 IrVersion, in
 		const TConstArrayView<uint8>& Data = InWeightTensorsData[Idx];
 		check(Data.Num() == Desc.GetDataSize());
 		ModelBuilderONNXHelper::BuildShapeForModel(false, Desc.GetShape(), ShapeForModel);
-		IModelBuilder::FHTensor Tensor = Builder->AddConstantTensor(Desc.GetName(), Desc.GetDataType(), ShapeForModel, Data.GetData(), Data.Num());
+		IModelBuilder::FHTensor Tensor = Desc.GetDataType() == ENNETensorDataType::None ?
+			  Builder->AddEmptyTensor()
+			: Builder->AddConstantTensor(Desc.GetName(), Desc.GetDataType(), ShapeForModel, Data.GetData(), Data.Num());
 
 		WeightTensors.Emplace(Tensor);
 	}
diff --git a/Engine/Plugins/Experimental/NNERuntimeRDG/Source/NNERuntimeRDGUtils/Private/NNERuntimeRDGUtilsModelOptimizerNNE.cpp b/Engine/Plugins/Experimental/NNERuntimeRDG/Source/NNERuntimeRDGUtils/Private/NNERuntimeRDGUtilsModelOptimizerNNE.cpp
--- a/Engine/Plugins/Experimental/NNERuntimeRDG/Source/NNERuntimeRDGUtils/Private/NNERuntimeRDGUtilsModelOptimizerNNE.cpp
+++ b/Engine/Plugins/Experimental/NNERuntimeRDG/Source/NNERuntimeRDGUtils/Private/NNERuntimeRDGUtilsModelOptimizerNNE.cpp
@@ -384,7 +384,7 @@ namespace ModelOptimizerNNEHelper
 				}
 				else // Empty tensor
 				{
-					Tensor = Builder->AddConstantTensor(/* Name */ TEXT(""), ENNETensorDataType::None, /* Shape */ {0}, /* Data */ nullptr, /* DataSize */ 0);
+					Tensor = Builder->AddEmptyTensor();
 				}
 
 				Builder->AddOperatorInput(Op, Tensor);
diff --git a/Engine/Source/Runtime/NNE/Public/NNETypes.h b/Engine/Source/Runtime/NNE/Public/NNETypes.h
--- a/Engine/Source/Runtime/NNE/Public/NNETypes.h
+++ b/Engine/Source/Runtime/NNE/Public/NNETypes.h
@@ -14,7 +14,7 @@
 UENUM()
 enum class ENNETensorDataType : uint8
 {
-	None,
+	None,								//!< For empty tensors (ENNEFormatTensorType::Empty)
 	Char,								//!< Character type
 	Boolean,							//!< Boolean type
 	Half,								//!< 16-bit floating number
