NNE - Add support for Hlsl Pad-13

Problems:
* RDG Hlsl backend lacks support for Pad-13.

Solutions:
* Extend RDG Hlsl implementation of Pad to version 13 and add tests.

Tests:
* Win64 Standalone tests: all
* Win64 Editor tests: all
* Win64 Editor tests: Pad opset tests

#rnx

#virtualized

diff --git a/Engine/Plugins/Experimental/NNERuntimeRDG/Source/NNERuntimeRDG/Private/Helper/NNERuntimeRDGOperatorHelper.cpp b/Engine/Plugins/Experimental/NNERuntimeRDG/Source/NNERuntimeRDG/Private/Helper/NNERuntimeRDGOperatorHelper.cpp
--- a/Engine/Plugins/Experimental/NNERuntimeRDG/Source/NNERuntimeRDG/Private/Helper/NNERuntimeRDGOperatorHelper.cpp
+++ b/Engine/Plugins/Experimental/NNERuntimeRDG/Source/NNERuntimeRDG/Private/Helper/NNERuntimeRDGOperatorHelper.cpp
@@ -4,8 +4,8 @@
 
 namespace UE::NNERuntimeRDG::Private::OperatorHelper
 {
-
-	bool GetInt32ArrayFromConstTensor(TArray<int32, TInlineAllocator<NNE::FTensorShape::MaxRank>>& Attr, const NNE::Internal::FTensorRef Tensor)
+	template<typename Allocator>
+	bool GetInt32ArrayFromConstTensorImpl(TArray<int32, Allocator>& Attr, const NNE::Internal::FTensorRef Tensor)
 	{
 		check(Tensor != nullptr);
 		Attr.Reset();
@@ -33,6 +33,16 @@ namespace UE::NNERuntimeRDG::Private::OperatorHelper
 
 		return false;
 	}
+	
+	bool GetInt32ArrayFromConstTensor(TArray<int32, TInlineAllocator<NNE::FTensorShape::MaxRank>>& Attr, const NNE::Internal::FTensorRef Tensor)
+	{
+		return GetInt32ArrayFromConstTensorImpl(Attr, Tensor);
+	}
+
+	bool GetInt32ArrayFromConstTensor(TArray<int32>& Attr, const NNE::Internal::FTensorRef Tensor)
+	{
+		return GetInt32ArrayFromConstTensorImpl(Attr, Tensor);
+	}
 
 } // UE::NNERuntimeRDG::Private::OperatorHelper
 
diff --git a/Engine/Plugins/Experimental/NNERuntimeRDG/Source/NNERuntimeRDG/Private/Helper/NNERuntimeRDGOperatorHelper.h b/Engine/Plugins/Experimental/NNERuntimeRDG/Source/NNERuntimeRDG/Private/Helper/NNERuntimeRDGOperatorHelper.h
--- a/Engine/Plugins/Experimental/NNERuntimeRDG/Source/NNERuntimeRDG/Private/Helper/NNERuntimeRDGOperatorHelper.h
+++ b/Engine/Plugins/Experimental/NNERuntimeRDG/Source/NNERuntimeRDG/Private/Helper/NNERuntimeRDGOperatorHelper.h
@@ -9,6 +9,7 @@
 namespace UE::NNERuntimeRDG::Private::OperatorHelper
 {
 	bool GetInt32ArrayFromConstTensor(TArray<int32, TInlineAllocator<NNE::FTensorShape::MaxRank>>& Attr, const NNE::Internal::FTensorRef Tensor);
+	bool GetInt32ArrayFromConstTensor(TArray<int32>& Attr, const NNE::Internal::FTensorRef Tensor);
 
 } // UE::NNERuntimeRDG::Private::OperatorHelper
 
diff --git a/Engine/Plugins/Experimental/NNERuntimeRDG/Source/NNERuntimeRDG/Private/Hlsl/NNERuntimeRDGPad.cpp b/Engine/Plugins/Experimental/NNERuntimeRDG/Source/NNERuntimeRDG/Private/Hlsl/NNERuntimeRDGPad.cpp
--- a/Engine/Plugins/Experimental/NNERuntimeRDG/Source/NNERuntimeRDG/Private/Hlsl/NNERuntimeRDGPad.cpp
+++ b/Engine/Plugins/Experimental/NNERuntimeRDG/Source/NNERuntimeRDG/Private/Hlsl/NNERuntimeRDGPad.cpp
@@ -3,6 +3,7 @@
 #include "NNERuntimeRDGPad.h"
 #include "NNEHlslShadersPadCS.h"
 #include "NNERuntimeRDGHlslHelper.h"
+#include "Helper/NNERuntimeRDGOperatorHelper.h"
 #include "NNETensor.h"
 #include "NNETypes.h"
 
@@ -15,6 +16,7 @@ namespace UE::NNERuntimeRDG::Private::Hlsl
 	/**
 	 * Pad operator implementation
 	 */
+	template<int Version>
 	class FPad : public FOperatorHlsl
 	{
 	public:
@@ -23,18 +25,55 @@ namespace UE::NNERuntimeRDG::Private::Hlsl
 		virtual ~FPad() = default;
 
 		TArray<int32> Pads;
-		float Value;
+		float Value = 0.0f;
 		EPadMode Mode;
 
 	public:
 
 		virtual int PrepareOutputs(TConstArrayView<NNE::Internal::FTensorRef> InputTensors, TArrayView<NNE::Internal::FTensorRef> OutputTensors) override
 		{
-			check(InputTensors.Num() == 1);
+			check(Version >= 11 ? (InputTensors.Num() >= 2 && InputTensors.Num() <= 3 ) : (InputTensors.Num() == 1));
 			check(OutputTensors.Num() == 1);
 
 			const NNE::Internal::FTensor& X = *InputTensors[0];
 
+			if constexpr (Version >= 11)
+			{
+				const NNE::Internal::FTensorRef PadsTensor = InputTensors[1];
+				if(!PadsTensor->HasPreparedData())
+				{
+					UE_LOG(LogNNE, Warning, TEXT("Pad input 'pads' (name: %s) should be constant."), *PadsTensor->GetName());
+					return -1;
+				}
+
+				OperatorHelper::GetInt32ArrayFromConstTensor(Pads, PadsTensor);
+
+				if(InputTensors.Num() == 3)
+				{
+					const NNE::Internal::FTensor& ValueTensor = *InputTensors[2];
+
+					if(!ValueTensor.HasPreparedData())
+					{
+						UE_LOG(LogNNE, Warning, TEXT("Pad input 'constant_value' (name: %s) should be constant."), *ValueTensor.GetName());
+						return -1;
+					}
+
+					if(ValueTensor.GetPreparedData<float>().Num() != 1)
+					{
+						UE_LOG(LogNNE, Warning, TEXT("Pad input 'constant_value' (name: %s) should be scalar, however it is not."), *ValueTensor.GetName());
+						return -1;
+					}
+
+					Value = ValueTensor.GetPreparedData<float>()[0];
+				}
+			}
+
+			if ((2*X.GetShape().Rank()) != Pads.Num())
+			{
+				UE_LOG(LogNNE, Warning, TEXT("pads attribute lenght (%d) should be twice the rank of input X (%d)."), Pads.Num(), X.GetShape().Rank());
+				return false;
+			}
+
 			TArray<uint32> OutputShapeData;
 			for (int32 i = 0; i < X.GetShape().Rank(); ++i)
 			{
@@ -59,18 +98,15 @@ namespace UE::NNERuntimeRDG::Private::Hlsl
 		{
 			using namespace UE::NNEHlslShaders::Internal;
 			
-			check(InputTensorDescs.Num() == 1);
+			check(Version >= 11 ? (InputTensorDescs.Num() >= 2 && InputTensorDescs.Num() <= 3 ) : (InputTensorDescs.Num() == 1));
 			check(OutputTensorDescs.Num() == 1);
 			
-			Pads = Attributes.GetValue<TArray<int32>>(TEXT("pads"));
-			Value = Attributes.GetValueOrDefault<float>(TEXT("value"), 0.0f);
-			FPadCS::LexFromString(Mode, *Attributes.GetValueOrDefault<FString>(TEXT("mode"), TEXT("constant")));
-
-			if ((2*InputTensorDescs[0].GetShape().Rank()) != Pads.Num())
+			if constexpr (Version < 11)
 			{
-				UE_LOG(LogNNE, Warning, TEXT("pads attribute lenght (%d) should be twice the rank of input X (%d)."), Pads.Num(), InputTensorDescs[0].GetShape().Rank());
-				return false;
+				Pads = Attributes.GetValue<TArray<int32>>(TEXT("pads"));
+				Value = Attributes.GetValueOrDefault<float>(TEXT("value"), 0.0f);
 			}
+			FPadCS::LexFromString(Mode, *Attributes.GetValueOrDefault<FString>(TEXT("mode"), TEXT("constant")));
 
 			return true;
 		}
@@ -79,7 +115,7 @@ namespace UE::NNERuntimeRDG::Private::Hlsl
 		{
 			using namespace UE::NNEHlslShaders::Internal;
 			
-			check(InputTensors.Num() == 1);
+			check(Version >= 11 ? (InputTensors.Num() >= 2 && InputTensors.Num() <= 3 ) : (InputTensors.Num() == 1));
 			check(OutputTensors.Num() == 1);
 			check(InputTensors[0] != nullptr);
 			check(OutputTensors[0] != nullptr);
@@ -125,33 +161,53 @@ namespace UE::NNERuntimeRDG::Private::Hlsl
 		}
 	};
 
+	template<int Version>
 	bool ValidatePadOperator(const NNE::FAttributeMap& AttributeMap, TConstArrayView<ENNETensorDataType> InputTypes, TConstArrayView<NNE::FSymbolicTensorShape> InputShapes)
 	{
 		bool bIsValid = true;
 
 		FAttributeValidator AttributeValidator;
 		AttributeValidator.AddOptional(TEXT("mode"), ENNEAttributeDataType::String);
-		AttributeValidator.AddRequired(TEXT("pads"), ENNEAttributeDataType::Int32Array);
-		AttributeValidator.AddOptional(TEXT("value"), ENNEAttributeDataType::Float);
+		if constexpr (Version < 11)
+		{
+			AttributeValidator.AddRequired(TEXT("pads"), ENNEAttributeDataType::Int32Array);
+			AttributeValidator.AddOptional(TEXT("value"), ENNEAttributeDataType::Float);
+		}
 		bIsValid &= AttributeValidator.Validate(AttributeMap);
 		
 		FInputValidator InputValidator;
+		InputValidator.SetTemplateCount(2);
 		InputValidator.AddSupportedType(ENNETensorDataType::Float);
+		InputValidator.AddSupportedType(ENNETensorDataType::Int64, /*TemplateIdx*/ 1);
 		InputValidator.AddRequired();
+		if constexpr (Version >= 11)
+		{
+			InputValidator.AddRequired(/*TemplateIdx*/ 1);
+			InputValidator.AddOptional();
+		}
+		
 		bIsValid &= InputValidator.Validate(InputTypes);
 
 		return bIsValid;
 	}
 
+	template<int Version>
 	FOperatorHlsl* CreatePadOperator()
 	{
-		return new FPad();
+		return new FPad<Version>();
 	}
 
 	bool RegisterPadOperator(FOperatorRegistryHlsl& Registry)
 	{
 		// Note: support of a particular version is partial with respect to tensor data types (only the most typical ones are usually supported).
-		Registry.OpAdd({{TEXT("Pad"), TEXT("Onnx")}, 2}, CreatePadOperator, ValidatePadOperator);
+		#define OP(Version) \
+		Registry.OpAdd({{TEXT("Pad"), TEXT("Onnx")}, Version}, CreatePadOperator<Version>, ValidatePadOperator<Version>);
+
+		OP(2)
+		OP(11)
+		OP(13)
+
+		#undef OP
 		return true;
 	}
 } // UE::NNERuntimeRDG::Private::Hlsl
