[NNERuntimeRDG HLSL] Adding support for pad operator version 18

Problem: Pad operator is only supported for version 13

Solution: Add support for optional pads 'axes' tensor input to make the operator compatible with version 18

Tested:
* Win64 Editor tests: all
* Win64 Standalone tests: all

#jra UE-217885
#rnx

#virtualized

diff --git a/Engine/Plugins/Experimental/NNERuntimeRDG/Source/NNERuntimeRDG/Private/Helper/NNERuntimeRDGOperatorHelper.cpp b/Engine/Plugins/Experimental/NNERuntimeRDG/Source/NNERuntimeRDG/Private/Helper/NNERuntimeRDGOperatorHelper.cpp
--- a/Engine/Plugins/Experimental/NNERuntimeRDG/Source/NNERuntimeRDG/Private/Helper/NNERuntimeRDGOperatorHelper.cpp
+++ b/Engine/Plugins/Experimental/NNERuntimeRDG/Source/NNERuntimeRDG/Private/Helper/NNERuntimeRDGOperatorHelper.cpp
@@ -38,6 +38,11 @@ namespace UE::NNERuntimeRDG::Private::OperatorHelper
 	{
 		return GetInt32ArrayFromConstTensorImpl(Attr, Tensor);
 	}
+	
+	bool GetInt32ArrayFromConstTensor(TArray<int32, TInlineAllocator<2 * NNE::FTensorShape::MaxRank>>& Attr, const NNE::Internal::FTensorRef Tensor)
+	{
+		return GetInt32ArrayFromConstTensorImpl(Attr, Tensor);
+	}
 
 	bool GetInt32ArrayFromConstTensor(TArray<int32>& Attr, const NNE::Internal::FTensorRef Tensor)
 	{
diff --git a/Engine/Plugins/Experimental/NNERuntimeRDG/Source/NNERuntimeRDG/Private/Helper/NNERuntimeRDGOperatorHelper.h b/Engine/Plugins/Experimental/NNERuntimeRDG/Source/NNERuntimeRDG/Private/Helper/NNERuntimeRDGOperatorHelper.h
--- a/Engine/Plugins/Experimental/NNERuntimeRDG/Source/NNERuntimeRDG/Private/Helper/NNERuntimeRDGOperatorHelper.h
+++ b/Engine/Plugins/Experimental/NNERuntimeRDG/Source/NNERuntimeRDG/Private/Helper/NNERuntimeRDGOperatorHelper.h
@@ -9,6 +9,7 @@
 namespace UE::NNERuntimeRDG::Private::OperatorHelper
 {
 	bool GetInt32ArrayFromConstTensor(TArray<int32, TInlineAllocator<NNE::FTensorShape::MaxRank>>& Attr, const NNE::Internal::FTensorRef Tensor);
+	bool GetInt32ArrayFromConstTensor(TArray<int32, TInlineAllocator<2 * NNE::FTensorShape::MaxRank>>& Attr, const NNE::Internal::FTensorRef Tensor);
 	bool GetInt32ArrayFromConstTensor(TArray<int32>& Attr, const NNE::Internal::FTensorRef Tensor);
 
 } // UE::NNERuntimeRDG::Private::OperatorHelper
diff --git a/Engine/Plugins/Experimental/NNERuntimeRDG/Source/NNERuntimeRDG/Private/Hlsl/NNERuntimeRDGPad.cpp b/Engine/Plugins/Experimental/NNERuntimeRDG/Source/NNERuntimeRDG/Private/Hlsl/NNERuntimeRDGPad.cpp
--- a/Engine/Plugins/Experimental/NNERuntimeRDG/Source/NNERuntimeRDG/Private/Hlsl/NNERuntimeRDGPad.cpp
+++ b/Engine/Plugins/Experimental/NNERuntimeRDG/Source/NNERuntimeRDG/Private/Hlsl/NNERuntimeRDGPad.cpp
@@ -32,23 +32,22 @@ namespace UE::NNERuntimeRDG::Private::Hlsl
 
 		virtual int PrepareOutputs(TConstArrayView<NNE::Internal::FTensorRef> InputTensors, TArrayView<NNE::Internal::FTensorRef> OutputTensors) override
 		{
-			check(Version >= 11 ? (InputTensors.Num() >= 2 && InputTensors.Num() <= 3 ) : (InputTensors.Num() == 1));
+			CheckInputTensorCount(InputTensors.Num());
 			check(OutputTensors.Num() == 1);
 
 			const NNE::Internal::FTensor& X = *InputTensors[0];
+			const int32 Rank = X.GetShape().Rank();
 
 			if constexpr (Version >= 11)
 			{
 				const NNE::Internal::FTensorRef PadsTensor = InputTensors[1];
-				if(!PadsTensor->HasPreparedData())
+				if(!PadsTensor->IsConstant())
 				{
 					UE_LOG(LogNNE, Warning, TEXT("Pad input 'pads' (name: %s) should be constant."), *PadsTensor->GetName());
 					return -1;
 				}
 
-				OperatorHelper::GetInt32ArrayFromConstTensor(Pads, PadsTensor);
-
-				if(InputTensors.Num() == 3)
+				if (InputTensors.Num() >= 3 && InputTensors[2]->GetDataType() != ENNETensorDataType::None)
 				{
 					const NNE::Internal::FTensor& ValueTensor = *InputTensors[2];
 
@@ -66,6 +65,65 @@ namespace UE::NNERuntimeRDG::Private::Hlsl
 
 					Value = ValueTensor.GetPreparedData<float>()[0];
 				}
+
+				if (InputTensors.Num() >= 4 && InputTensors[3]->GetDataType() != ENNETensorDataType::None)
+				{
+					const NNE::Internal::FTensorRef AxesTensor = InputTensors[3];
+
+					if(!AxesTensor->IsConstant())
+					{
+						UE_LOG(LogNNE, Warning, TEXT("Pad input 'axes' (name: %s) should be constant."), *AxesTensor->GetName());
+						return -1;
+					}
+
+					TArray<int32, TInlineAllocator<NNE::FTensorShape::MaxRank>> Axes;
+					TArray<int32, TInlineAllocator<2 * NNE::FTensorShape::MaxRank>> RelativePads;
+					if (AxesTensor->HasPreparedData())
+					{
+						OperatorHelper::GetInt32ArrayFromConstTensor(Axes, AxesTensor);
+					}
+
+					if (PadsTensor->HasPreparedData())
+					{
+						OperatorHelper::GetInt32ArrayFromConstTensor(RelativePads, PadsTensor);
+					}
+
+					if (RelativePads.Num() != Axes.Num() * 2)
+					{
+						UE_LOG(LogNNE, Warning, TEXT("Pad input 'axes' (name: %s) has to have a size that is twice the size \
+													  of pad input 'pads' (name: %s), but they have size %i and %i respectively."), 
+												*AxesTensor->GetName(), *PadsTensor->GetName(), Axes.Num(), RelativePads.Num());
+						return -1;
+					}
+
+					Pads.Init(0, 2 * Rank);
+					for (int AxesIndex = 0; AxesIndex < Axes.Num(); ++AxesIndex)
+					{
+						int32 Axis = Axes[AxesIndex];
+						if (Axis < -Rank || Axis >= Rank)
+						{
+							UE_LOG(LogNNE, Warning, TEXT("Pad input value at index %i of the 'axes' (name: %s) tensor \
+														  needs to be in the range [-Rank, Rank - 1], but value is %i with a rank of %i."), 
+													AxesIndex, *AxesTensor->GetName(), Axis, Rank);
+							return -1;
+						}
+						if (Axis < 0)
+						{
+							Axis = Rank + Axis;
+						}
+						Pads[Axis] = RelativePads[AxesIndex];
+						Pads[Axis + Rank] = RelativePads[AxesIndex + Axes.Num()];
+					}
+				}
+				else
+				{
+					if (!PadsTensor->HasPreparedData())
+					{
+						UE_LOG(LogNNE, Warning, TEXT("pads attribute lenght (%d) should be twice the rank of input X (%d)."), 0, X.GetShape().Rank());
+						return -1;
+					}
+					OperatorHelper::GetInt32ArrayFromConstTensor(Pads, PadsTensor);
+				}
 			}
 
 			if ((2*X.GetShape().Rank()) != Pads.Num())
@@ -98,7 +156,7 @@ namespace UE::NNERuntimeRDG::Private::Hlsl
 		{
 			using namespace UE::NNEHlslShaders::Internal;
 			
-			check(Version >= 11 ? (InputTensorDescs.Num() >= 2 && InputTensorDescs.Num() <= 3 ) : (InputTensorDescs.Num() == 1));
+			CheckInputTensorCount(InputTensorDescs.Num());
 			check(OutputTensorDescs.Num() == 1);
 			
 			if constexpr (Version < 11)
@@ -115,7 +173,7 @@ namespace UE::NNERuntimeRDG::Private::Hlsl
 		{
 			using namespace UE::NNEHlslShaders::Internal;
 			
-			check(Version >= 11 ? (InputTensors.Num() >= 2 && InputTensors.Num() <= 3 ) : (InputTensors.Num() == 1));
+			CheckInputTensorCount(InputTensors.Num());
 			check(OutputTensors.Num() == 1);
 			check(InputTensors[0] != nullptr);
 			check(OutputTensors[0] != nullptr);
@@ -159,6 +217,24 @@ namespace UE::NNERuntimeRDG::Private::Hlsl
 				Params,
 				ThreadGroupCount);
 		}
+
+	private:
+		void CheckInputTensorCount(const int32 count)
+		{
+			if constexpr (Version < 11)
+			{
+				check(count == 1);
+			}
+			else if constexpr (Version < 18)
+			{
+				check(count >= 2 && count <= 3);
+			}
+			else
+			{
+				check(count >= 2 && count <= 4);
+			}
+		}
+
 	};
 
 	template<int Version>
@@ -176,15 +252,21 @@ namespace UE::NNERuntimeRDG::Private::Hlsl
 		bIsValid &= AttributeValidator.Validate(AttributeMap);
 		
 		FInputValidator InputValidator;
-		InputValidator.SetTemplateCount(2);
+		InputValidator.SetTemplateCount(3);
 		InputValidator.AddSupportedType(ENNETensorDataType::Float);
 		InputValidator.AddSupportedType(ENNETensorDataType::Int64, /*TemplateIdx*/ 1);
+		InputValidator.AddSupportedType(ENNETensorDataType::Int32, /*TemplateIdx*/ 2);
+		InputValidator.AddSupportedType(ENNETensorDataType::Int64, /*TemplateIdx*/ 2);
 		InputValidator.AddRequired();
 		if constexpr (Version >= 11)
 		{
 			InputValidator.AddRequired(/*TemplateIdx*/ 1);
 			InputValidator.AddOptional();
 		}
+		if constexpr (Version >= 18)
+		{
+			InputValidator.AddOptional(/*TemplateIdx*/ 2);
+		}
 		
 		bIsValid &= InputValidator.Validate(InputTypes);
 
@@ -206,8 +288,10 @@ namespace UE::NNERuntimeRDG::Private::Hlsl
 		OP(2)
 		OP(11)
 		OP(13)
+		OP(18)
 
 		#undef OP
 		return true;
 	}
+
 } // UE::NNERuntimeRDG::Private::Hlsl
diff --git a/Engine/Plugins/Experimental/NNERuntimeRDG/Source/NNERuntimeRDG/Private/NNERuntimeRDGBase.cpp b/Engine/Plugins/Experimental/NNERuntimeRDG/Source/NNERuntimeRDG/Private/NNERuntimeRDGBase.cpp
--- a/Engine/Plugins/Experimental/NNERuntimeRDG/Source/NNERuntimeRDG/Private/NNERuntimeRDGBase.cpp
+++ b/Engine/Plugins/Experimental/NNERuntimeRDG/Source/NNERuntimeRDG/Private/NNERuntimeRDGBase.cpp
@@ -56,7 +56,8 @@ bool FInputValidator::Validate(TConstArrayView<ENNETensorDataType> InputTypes)
 		const int32 TemplateIdx = InputTemplateIndices[Idx];
 		
 		check(TemplateIdx < TemplateTypes.Num());
-		if (!TemplateTypes[TemplateIdx].Contains(InputTypes[Idx]))
+		if (!TemplateTypes[TemplateIdx].Contains(InputTypes[Idx]) &&
+			!(InputTypes[Idx] == ENNETensorDataType::None && Idx >= NumRequiredInput))
 		{
 			FString TargetType = LogHelper::GetTensorDataTypeName(InputTypes[Idx]);
 			UE_LOG(LogNNE, Warning, TEXT("Input at index '%d' (from template T%d) is of type '%s' witch is not supported for that input."), Idx, TemplateIdx, *TargetType);
