NNEDenoiser: Mark NNE runtime plugins as ClientOnly

These should not be necessary for Server builds. It appears that the ClientOnly flag from the parent is not propagated to dependencies in incremental builds.

#rnx

diff --git a/Engine/Plugins/Experimental/NNERuntimeRDG/NNERuntimeRDG.uplugin b/Engine/Plugins/Experimental/NNERuntimeRDG/NNERuntimeRDG.uplugin
--- a/Engine/Plugins/Experimental/NNERuntimeRDG/NNERuntimeRDG.uplugin
+++ b/Engine/Plugins/Experimental/NNERuntimeRDG/NNERuntimeRDG.uplugin
@@ -21,13 +21,13 @@
 	"Modules": [
 		{
 			"Name": "NNEHlslShaders",
-			"Type": "RuntimeAndProgram",
+			"Type": "ClientOnly",
 			"LoadingPhase": "PostConfigInit",
 			"PlatformAllowList": [ "Win64", "Linux", "Mac" ]
 		},
 		{
 			"Name": "NNERuntimeRDG",
-			"Type": "RuntimeAndProgram",
+			"Type": "ClientOnly",
 			"LoadingPhase": "Default",
 			"PlatformAllowList": [ "Win64", "Linux", "Mac" ]
 		},
diff --git a/Engine/Plugins/NNE/NNERuntimeORT/NNERuntimeORT.uplugin b/Engine/Plugins/NNE/NNERuntimeORT/NNERuntimeORT.uplugin
--- a/Engine/Plugins/NNE/NNERuntimeORT/NNERuntimeORT.uplugin
+++ b/Engine/Plugins/NNE/NNERuntimeORT/NNERuntimeORT.uplugin
@@ -18,7 +18,7 @@
 	"Modules": [
 		{
 			"Name": "NNERuntimeORT",
-			"Type": "RuntimeAndProgram",
+			"Type": "ClientOnly",
 			"LoadingPhase": "PreDefault",
 			"PlatformAllowList": [ "Win64", "Linux", "Mac" ]
 		}
