[NNE] Revert 34447957, mark plugin modules as RuntimeAndProgram again (was set to ClientOnly).

#tests

#rnx

diff --git a/Engine/Plugins/Experimental/NNERuntimeRDG/NNERuntimeRDG.uplugin b/Engine/Plugins/Experimental/NNERuntimeRDG/NNERuntimeRDG.uplugin
--- a/Engine/Plugins/Experimental/NNERuntimeRDG/NNERuntimeRDG.uplugin
+++ b/Engine/Plugins/Experimental/NNERuntimeRDG/NNERuntimeRDG.uplugin
@@ -21,13 +21,13 @@
 	"Modules": [
 		{
 			"Name": "NNEHlslShaders",
-			"Type": "ClientOnly",
+			"Type": "RuntimeAndProgram",
 			"LoadingPhase": "PostConfigInit",
 			"PlatformAllowList": [ "Win64", "Linux", "Mac" ]
 		},
 		{
 			"Name": "NNERuntimeRDG",
-			"Type": "ClientOnly",
+			"Type": "RuntimeAndProgram",
 			"LoadingPhase": "Default",
 			"PlatformAllowList": [ "Win64", "Linux", "Mac" ]
 		},
diff --git a/Engine/Plugins/NNE/NNERuntimeORT/NNERuntimeORT.uplugin b/Engine/Plugins/NNE/NNERuntimeORT/NNERuntimeORT.uplugin
--- a/Engine/Plugins/NNE/NNERuntimeORT/NNERuntimeORT.uplugin
+++ b/Engine/Plugins/NNE/NNERuntimeORT/NNERuntimeORT.uplugin
@@ -18,7 +18,7 @@
 	"Modules": [
 		{
 			"Name": "NNERuntimeORT",
-			"Type": "ClientOnly",
+			"Type": "RuntimeAndProgram",
 			"LoadingPhase": "PreDefault",
 			"PlatformAllowList": [ "Win64", "Linux", "Mac" ]
 		}
