[NNE] NNERuntimeHLSL Move shader file to follow engine pattern.

Problem: Shader file where placed in /source folder breaking server build
Solution: Move shader file at the plugin level, following common engine pattern.

#test all nne tests on win64 editor and packaged win64 standalone, smoke test and style transfer test on console
#rnx

diff --git a/Engine/Plugins/Experimental/NNERuntimeRDG/Shaders/Private/NNEHlslShaders/NNEHlslShadersBatchNormalization.usf b/Engine/Plugins/Experimental/NNERuntimeRDG/Shaders/Private/NNEHlslShaders/NNEHlslShadersBatchNormalization.usf
new file mode 100644
--- /dev/null
+++ b/Engine/Plugins/Experimental/NNERuntimeRDG/Shaders/Private/NNEHlslShaders/NNEHlslShadersBatchNormalization.usf
@@ -0,0 +1,41 @@
+// Copyright Epic Games, Inc. All Rights Reserved.
+
+#include "/Engine/Public/Platform.ush"
+
+#define WORK_TYPE float
+
+Buffer<WORK_TYPE> X;
+Buffer<WORK_TYPE> Scales;
+Buffer<WORK_TYPE> Bias;
+Buffer<WORK_TYPE> Mean;
+Buffer<WORK_TYPE> Var;
+RWBuffer<WORK_TYPE> Output;
+float Epsilon;
+uint SpatialVolume;
+uint DimC;
+uint ThreadCountX;
+uint Num;
+
+WORK_TYPE ApplyBatchNormalization(WORK_TYPE x, WORK_TYPE input_mean, WORK_TYPE input_var, WORK_TYPE scale, WORK_TYPE bias)
+{
+	//[REDACTED_PROFILE_URL]/onnx/blob/main/docs/Operators.md#BatchNormalization
+	return (x - input_mean) / sqrt(input_var + Epsilon) * scale + bias;
+
+}
+
+
+[numthreads(THREADGROUP_SIZE_X, 1, 1)]
+void BatchNormalization(in const uint3 DispatchThreadID : SV_DispatchThreadID)
+{
+	const uint IndexX = DispatchThreadID.y * ThreadCountX + DispatchThreadID.x;
+	const uint IndexOther = (IndexX / SpatialVolume) % DimC;
+	if (IndexX < Num)
+	{
+		float x = X[IndexX];
+		float mean = Mean[IndexOther];
+		float var = Var[IndexOther];
+		float scale = Scales[IndexOther];
+		float bias = Bias[IndexOther];
+		Output[IndexX] = ApplyBatchNormalization(x, mean, var, scale, bias);
+	}
+}
diff --git a/Engine/Plugins/Experimental/NNERuntimeRDG/Shaders/Private/NNEHlslShaders/NNEHlslShadersBroadcastHelper.ush b/Engine/Plugins/Experimental/NNERuntimeRDG/Shaders/Private/NNEHlslShaders/NNEHlslShadersBroadcastHelper.ush
new file mode 100644
--- /dev/null
+++ b/Engine/Plugins/Experimental/NNERuntimeRDG/Shaders/Private/NNEHlslShaders/NNEHlslShadersBroadcastHelper.ush
@@ -0,0 +1,7 @@
+// Copyright Epic Games, Inc. All Rights Reserved.
+
+void DivMod(uint op1, uint op2, out uint q, out uint r)
+{
+	q = op1 / op2;
+	r = op1 - (q * op2);
+}
\ No newline at end of file
diff --git a/Engine/Plugins/Experimental/NNERuntimeRDG/Shaders/Private/NNEHlslShaders/NNEHlslShadersConv.usf b/Engine/Plugins/Experimental/NNERuntimeRDG/Shaders/Private/NNEHlslShaders/NNEHlslShadersConv.usf
new file mode 100644
--- /dev/null
+++ b/Engine/Plugins/Experimental/NNERuntimeRDG/Shaders/Private/NNEHlslShaders/NNEHlslShadersConv.usf
@@ -0,0 +1,304 @@
+// Copyright Epic Games, Inc. All Rights Reserved.
+
+#include "/Engine/Public/Platform.ush"
+
+#define WORK_TYPE float
+#define BUFFER_TYPE float
+#define BUFFER_TO_WORK_TYPE(x) x
+#define WORK_TO_BUFFER_TYPE(x) x
+
+// Must correspond to EConvGroupSize defined in NNEHlslShadersConvCS.h
+#if GROUP_SIZE == 0
+#define NUM_GROUP_THREADS 128
+#elif GROUP_SIZE == 1
+#define NUM_GROUP_THREADS 256
+#elif GROUP_SIZE == 2
+#define NUM_GROUP_THREADS 512
+#endif
+
+Buffer<BUFFER_TYPE> X; // N x C x XD0 x ... x XDi
+Buffer<BUFFER_TYPE> W; // M x C/GROUPS x WD0 x ... x WDi OR WD0 x WD1 x C X M if WEIGHTS_TRANSPOSED (only supported for 4D)
+RWBuffer<BUFFER_TYPE> Y;  // N x M x YD0 x ... x YDi
+Buffer<BUFFER_TYPE> B; // B0, ..., BM
+
+// x: Dilation0, ...., Dilationi
+// y: Strides0, ...., Stridesi
+// z: Pad0Begin, ...., PadiBegin
+int4 Dilation_Stride_XBlockStartOffset_DilationXBlockStride[MAX_NUM_DIMENSIONS];
+
+// x: Element j contains PROD(GDj+1,..,GDi) with the last element containing 1, with GDi being the number of thread groups contained by a dimension i
+// y: Number of threads in each dimension of a group
+// z: Element j contains PROD(GTDj+1,..,GTDi) with the last element containing 1, with GTDi being the number of threads contained by a dimension i inside a group
+int4 GroupStride_GroupShape_GroupThreadStride_StrideXBlockStride[MAX_NUM_DIMENSIONS];
+
+// x: YD0, ..., YDi
+// y: Element j contains PROD(YDj+1,..,YDi) with the last element containing 1
+// z: XD0, ..., XDi
+// w: Element j contains PROD(XDj+1,..,XDi) with the last element containing 1
+int4 YDimension_YMemoryStride_XDimension_XMemoryStride[MAX_NUM_DIMENSIONS];
+
+// x: Number of elements in each dimension of a X block
+// y: The strides of the X block to be loaded
+// z: WD0, ..., WDi
+int4 XBlockStartStride_XBlockStride_WDimension_WDimensionDilationXBlockStride[MAX_NUM_DIMENSIONS];
+
+// x: 1/GroupStride
+// y: 1/GroupThreadStride
+// z: 1/XBlockStride
+float4 OneDiv_GroupStride_GroupThreadStride_XBlockStride[MAX_NUM_DIMENSIONS];
+
+int NumWFeatures; // Number of features (aka output channels)
+int NumWChannels; // Number of input channels
+
+int YBatchStride;
+int YOutputKernelStride;
+
+int XBatchStride; // The number of elements in each bacth of X
+int XChannelStride; // The number of elements in each channel of X
+
+int XBlockSize; // The total number of elements in each block to be loaded
+
+int NumChannelBatches; // The number of iteration needed to cover all channels that needs to be processed. Equals to ceil(NumWChannels/NumChannelsPerBatch)
+int NumChannelsPerBatch; // The number of complete kernel channels that can be loaded in one batch
+
+int WOutputKernelStride; // The total number of W elements per output kernel and thus PROD(C/GROUPS, WD0, .., WDi)
+int WChannelBatchSize; // The total number of W elements inside a channel batch and thus the number of channels per batch times PROD(WD0, .., WDi)
+int WChannelSize; // Number of elements in a single channel of W and thus PROD(WDi)
+
+float GroupsDivM; // The number of groups divided by the number of output kernels used to compute the start channel offset
+
+groupshared WORK_TYPE SharedMemoryX[NUM_GROUP_THREADS<<NUM_READS_PER_THREAD_POW2];
+groupshared WORK_TYPE SharedMemoryW[NUM_GROUP_THREADS];
+
+// Must correspond to EConvAlgorithm defined in NNEHlslShadersConvCS.h
+#if ALGORITHM < 1
+
+[numthreads(NUM_GROUP_THREADS, 1, 1)]
+void Conv(in const uint3 DispatchThreadID : SV_DispatchThreadID, in const uint3 GroupID : SV_GroupID, in const uint3 GroupThreadID : SV_GroupThreadID)
+{
+	// Group Z dimension and Y dimension are used for batch and output kernel
+    // This works as each group only works on the inner dimensions (num threads in y and z set to one, group count to N and M respectively)
+	int BatchIndex = GroupID.z;
+	int OutputKernelIndex = GroupID.y;
+
+	// Input channels and output kernels are divided into groups
+	// Thus the channel offset for this thread's output kernel can be computed as follows:
+	// The output kernel index is divided by the number of output kernels per group giving the group index
+	// This is multiplied by the number of channels in W to get the start index
+	int InputChannelStartIndex = NumWChannels * (int)(GroupsDivM * (float)OutputKernelIndex);
+
+	// Compute the output memory index and dimension indices of this thread inside its group as well as the start dimension indices of the X block to load
+	int YIndex = BatchIndex * YBatchStride + OutputKernelIndex * YOutputKernelStride;
+	int GroupThreadDimensionIndex[NUM_DIMENSIONS];// This will be used later to lookup the shared W and X to compute the final result
+    int XBlockStartDimensionIndex[NUM_DIMENSIONS];
+	int GroupIndex = GroupID.x;
+	int GroupThreadIndex = GroupThreadID.x;
+	UNROLL
+	for (int i = 0; i < NUM_DIMENSIONS; i++)
+	{
+        // Divide by the number of groups contained by one index of this dimension to get the dimension index n being the nth group in this dimension
+		int GroupDimensionIndex = (int)(OneDiv_GroupStride_GroupThreadStride_XBlockStride[i].x * (float)GroupIndex);
+		GroupIndex -= GroupDimensionIndex * GroupStride_GroupShape_GroupThreadStride_StrideXBlockStride[i].x;
+
+        // Same here with number of threads contained inside the group giving the nth thread inside the group in this dimension
+		GroupThreadDimensionIndex[i] = (int)(OneDiv_GroupStride_GroupThreadStride_XBlockStride[i].y * (float)GroupThreadIndex);
+		GroupThreadIndex -= GroupThreadDimensionIndex[i] * GroupStride_GroupShape_GroupThreadStride_StrideXBlockStride[i].z;
+		
+        // Compose the global index by the group index times group shape plus thread index inside the group
+		int ThreadDimensionIndex = GroupDimensionIndex * GroupStride_GroupShape_GroupThreadStride_StrideXBlockStride[i].y + GroupThreadDimensionIndex[i];
+        
+        // Subtract the padding from the index of the first element of the group to get the start index of the block needed by this group
+        XBlockStartDimensionIndex[i] = GroupDimensionIndex * XBlockStartStride_XBlockStride_WDimension_WDimensionDilationXBlockStride[i].x + Dilation_Stride_XBlockStartOffset_DilationXBlockStride[i].z;
+
+        // ThreadDimensionIndex is always positive but needs to be tested if it does not overflow
+		if (ThreadDimensionIndex >= YDimension_YMemoryStride_XDimension_XMemoryStride[i].x || YIndex < 0)
+		{
+			YIndex = -1;
+		}
+		else
+		{
+			YIndex += ThreadDimensionIndex * YDimension_YMemoryStride_XDimension_XMemoryStride[i].y;
+		}
+	}
+
+    // Compute the memory indices in global X which the current thread has to load as well as the indices in shared X where to write
+	int GlobalXIndex[0x1<<NUM_READS_PER_THREAD_POW2];
+	int SharedXIndex[0x1<<NUM_READS_PER_THREAD_POW2];
+	int ReadIdx = GroupThreadID.x;
+    UNROLL
+    for (int j = 0; j < (int)(0x1<<NUM_READS_PER_THREAD_POW2); j++)
+	{
+        // If the read index exceeded the number of elements in the X block further computation can be skipped
+        if (ReadIdx >= XBlockSize)
+        {
+            // Assign -1 for both to indicate that the element is not used and thus neither has to be read from global nor written to shared memory
+            GlobalXIndex[j] = -1;
+            SharedXIndex[j] = -1;
+            continue;
+        }
+
+        // Turn the flat group index multiple into thread dimension indices inside the volume to be loaded by this group
+        int TmpGlobalXIndex = BatchIndex * XBatchStride + InputChannelStartIndex * XChannelStride;
+        int TmpSharedXIndex = 0;
+        int TempReadIdx = ReadIdx;
+        UNROLL
+        for (int k = 0; k < NUM_DIMENSIONS; k++)
+        {
+            int DimensionIndex = (int)(OneDiv_GroupStride_GroupThreadStride_XBlockStride[k].z * (float)TempReadIdx);
+            int DimensionIndexOffset = DimensionIndex * XBlockStartStride_XBlockStride_WDimension_WDimensionDilationXBlockStride[k].y;
+		    TempReadIdx -= DimensionIndexOffset;
+
+            // Update TmpGlobalXIndex only if positive (it is already outside the volume if negative)
+            if (TmpGlobalXIndex >= 0)
+            {
+                // Check if the global X index is in range
+                int XIndex = XBlockStartDimensionIndex[k] + DimensionIndex;
+                if (XIndex < 0 || XIndex >= YDimension_YMemoryStride_XDimension_XMemoryStride[k].z)
+                {
+                    TmpGlobalXIndex = -1;
+                }
+                else
+                {
+                    TmpGlobalXIndex += XIndex * YDimension_YMemoryStride_XDimension_XMemoryStride[k].w;
+                }
+            }
+
+            // Update TmpSharedXIndex in any case as a zero has to be written
+            TmpSharedXIndex += DimensionIndexOffset;
+        }
+
+        // Assign the result, which will be -1 for the global and positive for the shared index if the element lies in the padding area
+        // This indicates that the global element should not be loaded but a zero has to be written to shared memory 
+        GlobalXIndex[j] = TmpGlobalXIndex;
+        SharedXIndex[j] = TmpSharedXIndex;
+
+        // Increment the read idx by the number of threads inside the group to cover all elements
+        ReadIdx += NUM_GROUP_THREADS;
+    }
+
+    // Write the zeros once to avoid rewriting the same value later
+    UNROLL
+    for (int l = 0; l < (int)(0x1<<NUM_READS_PER_THREAD_POW2); l++)
+	{
+        // If the global index is out of bound (padding area) but the shared index is valid (inside the X block) the shared element needs initialization to zero
+        if (GlobalXIndex[l] < 0 && SharedXIndex[l] >= 0)
+        {
+            SharedMemoryX[SharedXIndex[l]] = 0;
+        }
+    }
+
+    // Iterate over all channel batches
+    int WChannelOutputKernelOffset = OutputKernelIndex * WOutputKernelStride;
+    int WChannelBatchOffset = 0;
+    int WChannelIndex = 0;
+#if HAS_B == 0
+    WORK_TYPE Result = 0;
+#else
+    WORK_TYPE Result = BUFFER_TO_WORK_TYPE(B[OutputKernelIndex]);
+#endif
+    for (int m = 0; m < NumChannelBatches; m++)
+	{
+        // Load all channels of W belonging to this batch
+        int WIndex = WChannelBatchOffset + GroupThreadID.x;
+        if (GroupThreadID.x < WChannelBatchSize && WIndex < WOutputKernelStride)
+        {
+			#if WEIGHTS_TRANSPOSED
+				int MCHWIndex = WChannelOutputKernelOffset + WIndex;
+				int Whw = MCHWIndex % WChannelSize;
+				int Wc = (MCHWIndex / WChannelSize) % NumWChannels;
+				int Wm = (MCHWIndex / (NumWChannels * WChannelSize));
+				int HWCMIndex = Whw * NumWChannels * NumWFeatures + Wc * NumWFeatures + Wm;
+				SharedMemoryW[GroupThreadID.x] = BUFFER_TO_WORK_TYPE(W[HWCMIndex]);
+			#else
+				SharedMemoryW[GroupThreadID.x] = BUFFER_TO_WORK_TYPE(W[WChannelOutputKernelOffset + WIndex]);
+			#endif
+            WChannelBatchOffset += WChannelBatchSize;
+		}
+        
+        // Each channel batch iterates over all consecutive elements of the loaded kernel channels
+        // This results in bank conflict free lookups, as all threads of the same warp read exactly the same element which is then broadcasted
+        int SharedWIndex = 0;
+
+        // Iterate over all channels in this batch
+        for (int n = 0; n < NumChannelsPerBatch; n++)
+	    {
+            // Check if the channel index is valid as channel batches may not be aligned with the effective number of channels
+            if (WChannelIndex >= NumWChannels)
+            {
+                break;
+            }
+            WChannelIndex++;
+            
+            // Load the shared X
+            UNROLL
+            for (int o = 0; o < (int)(0x1<<NUM_READS_PER_THREAD_POW2); o++)
+	        {
+                // If the global index is negative, a zero has already been written, if shared index is negative there is no need to load anything
+                if (GlobalXIndex[o] >= 0 && SharedXIndex[o] >= 0)
+                {
+                    // Load the global data into shared memory
+                    SharedMemoryX[SharedXIndex[o]] = BUFFER_TO_WORK_TYPE(X[GlobalXIndex[o]]);
+
+                    // Update the indices for the next iteration by adding the channel stride
+                    GlobalXIndex[o] += XChannelStride;
+                }
+            }
+
+            // Sync the group before doing any computation
+            GroupMemoryBarrierWithGroupSync();
+
+            // Initialize a dimension index array used to iterate through the kernel dimensions
+            // Also initialize the start index in shared memory based on the result position
+            int WDimensionIndex[NUM_DIMENSIONS];
+            int SharedXIndex = 0;
+            UNROLL
+            for (int p = 0; p < NUM_DIMENSIONS; p++)
+            {
+                WDimensionIndex[p] = 0;
+
+                // The index of the first element (e.g. upper left) required by each output is offsetted by the kernel stride and then multiplied by the stride in memory
+                SharedXIndex += GroupThreadDimensionIndex[p] * GroupStride_GroupShape_GroupThreadStride_StrideXBlockStride[p].w;
+            }
+
+            // Compute and accumulate the result
+            for (int q = 0; q < WChannelSize; q++)
+            {
+                Result += SharedMemoryX[SharedXIndex] * SharedMemoryW[SharedWIndex];
+
+                // Advance the inner-most dimension
+                SharedWIndex++;
+                SharedXIndex += Dilation_Stride_XBlockStartOffset_DilationXBlockStride[NUM_DIMENSIONS-1].x;
+                WDimensionIndex[NUM_DIMENSIONS-1]++;
+
+                // Advance the other dimensions if necessary
+                UNROLL
+                for (int r = NUM_DIMENSIONS-1; r > 0; r--)
+                {
+                    if(WDimensionIndex[r] >= XBlockStartStride_XBlockStride_WDimension_WDimensionDilationXBlockStride[r].z)
+                    {
+                        // Adjust the SharedXIndex by resetting the current dimension and increase the next outer one
+                        SharedXIndex -= XBlockStartStride_XBlockStride_WDimension_WDimensionDilationXBlockStride[r].w;
+                        SharedXIndex += Dilation_Stride_XBlockStartOffset_DilationXBlockStride[r-1].w;
+                        WDimensionIndex[r] = 0;
+                        WDimensionIndex[r-1]++;
+                    }
+                    else 
+                    {
+                        break;
+                    }
+                }
+            }
+
+            // Sync the group before continuing any load
+            GroupMemoryBarrierWithGroupSync();
+        }
+    }
+
+    // Write the final result
+    if (YIndex >= 0)
+    {
+        Y[YIndex] = WORK_TO_BUFFER_TYPE(Result);
+    }
+}
+
+#endif
\ No newline at end of file
diff --git a/Engine/Plugins/Experimental/NNERuntimeRDG/Shaders/Private/NNEHlslShaders/NNEHlslShadersConvMatmul.usf b/Engine/Plugins/Experimental/NNERuntimeRDG/Shaders/Private/NNEHlslShaders/NNEHlslShadersConvMatmul.usf
new file mode 100644
--- /dev/null
+++ b/Engine/Plugins/Experimental/NNERuntimeRDG/Shaders/Private/NNEHlslShaders/NNEHlslShadersConvMatmul.usf
@@ -0,0 +1,223 @@
+// Copyright Epic Games, Inc. All Rights Reserved.
+
+#include "/Engine/Public/Platform.ush"
+
+#define WORK_TYPE float
+#define BUFFER_TYPE float
+#define BUFFER_TO_WORK_TYPE(x) x
+#define WORK_TO_BUFFER_TYPE(x) x
+
+#define LDS_IC 16 //Input channel cached in LDS
+#define TG_OC 32 //Number of output channel covered by a thread group
+#define TG_OP 32 //Number of output pixel covered by a thread group
+
+//To avoid bank conflict when writing out to DDR
+#define LDS_PADX 1
+#define LDS_PADY 2
+#define LDS_X_PADDEDSIZE ((4 + LDS_PADX) * 8 + LDS_PADY) //42
+
+groupshared WORK_TYPE LDSBuffer[32 * LDS_X_PADDEDSIZE]; //max (16*32+32,32*42) = 1344 WORK_TYPE 
+
+Buffer<BUFFER_TYPE> Input; // Ni x Ci x Hi x Wi
+Buffer<BUFFER_TYPE> Weight; // Cw x Ci x Hw x Ww OR Hw x Ww x Ci x Cw if WEIGHTS_TRANSPOSED
+Buffer<BUFFER_TYPE> Bias; // Cw
+RWBuffer<BUFFER_TYPE> Output; // Ni x Cw x Ho x Wo
+
+int Ci;
+int Hi;
+int Wi;
+int Ho;
+int Wo;
+int Cw;
+int Hw;
+int Ww;
+int PadLeft;
+int PadTop;
+int StrideH;
+int StrideW;
+
+WORK_TYPE ReadLDSWeights(int inputChannel, int outputChannel)
+{
+	return LDSBuffer[LDS_IC * TG_OP + inputChannel * TG_OC + outputChannel];
+}
+
+void WriteLDSWeights(int inputChannel, int outputChannel, WORK_TYPE value)
+{
+	LDSBuffer[LDS_IC * TG_OP + inputChannel * TG_OC + outputChannel] = value;
+}
+
+WORK_TYPE ReadLDSInputs(int inputChannel, int outputPixel)
+{
+	return LDSBuffer[inputChannel * TG_OP + outputPixel];
+}
+
+void WriteLDSInputs(int inputChannel, int outputPixel, WORK_TYPE value)
+{
+	LDSBuffer[inputChannel * TG_OP + outputPixel] = value;
+}
+
+[numthreads(8, 8, 1)]
+void ConvMatmul(
+	in const uint3 DispatchThreadID : SV_DispatchThreadID,
+	in const uint3 GroupID : SV_GroupID,
+	in const uint3 GroupThreadID : SV_GroupThreadID,
+	in const uint GroupIndex : SV_GroupIndex)
+{
+	const int DispatchThreadOutputPixelOffset = 4 * DispatchThreadID.x;
+	const int DispatchThreadOutputChannelOffset = 4 * DispatchThreadID.y;
+	const int GroupThreadOutputPixelOffset = 4 * GroupThreadID.x;
+	const int GroupThreadOutputChannelOffset = 4 * GroupThreadID.y;
+	const int GroupThreadIdx = GroupIndex;
+	const int Scalar_GroupOutputPixelOffset = TG_OP * GroupID.x;
+	const int Scalar_GroupOutputChannelOffset = TG_OC * GroupID.y;
+	const int Scalar_BatchOutputOffset = Ho * Wo * Cw * GroupID.z;
+	const int Scalar_BatchInputOffset = Hi * Wi * Ci * GroupID.z;
+	const int Scalar_GroupOutputPixelH = (Scalar_GroupOutputPixelOffset) / Wo * StrideH;
+	const int Scalar_GroupOutputPixelWBase = Scalar_GroupOutputPixelOffset % Wo * StrideW;
+	const int Scalar_GroupInputPixelKernelTopLeftH = Scalar_GroupOutputPixelH - PadTop;
+	const int Scalar_GroupInputPixelKernelTopLeftW = Scalar_GroupOutputPixelWBase - PadLeft;
+
+	const int cw = DispatchThreadOutputChannelOffset;
+	const int pi = DispatchThreadOutputPixelOffset;
+
+#if HAS_BIAS
+	WORK_TYPE BiasC0 = BUFFER_TO_WORK_TYPE(Bias[cw+0]);
+	WORK_TYPE BiasC1 = BUFFER_TO_WORK_TYPE(Bias[cw+1]);
+	WORK_TYPE BiasC2 = BUFFER_TO_WORK_TYPE(Bias[cw+2]);
+	WORK_TYPE BiasC3 = BUFFER_TO_WORK_TYPE(Bias[cw+3]);
+#else
+	WORK_TYPE BiasC0 = 0.0f;
+	WORK_TYPE BiasC1 = 0.0f;
+	WORK_TYPE BiasC2 = 0.0f;
+	WORK_TYPE BiasC3 = 0.0f;
+#endif
+	
+	int i;
+	// First index is thread output channel offset (on cw)
+	// 2nd index is thread pixels offset (on pi)
+	WORK_TYPE Values[4][4];
+	
+	UNROLL
+	for (i = 0; i < 4; ++i)
+	{
+		Values[0][i] = BiasC0;
+		Values[1][i] = BiasC1;
+		Values[2][i] = BiasC2;
+		Values[3][i] = BiasC3;
+	}
+	
+	int Scalar_KernelIdx = 0;
+	const int InputChannelOffset = GroupThreadIdx >= 32 ? LDS_IC/2 : 0;
+	const int OutputPixelOffset = GroupThreadIdx % 32;
+	const int OutputChannelOffset = GroupThreadIdx % 32;
+	
+	for (int hw = 0; hw < Hw; ++hw)
+	{
+		const int Scalar_InputPixelH = Scalar_GroupInputPixelKernelTopLeftH + hw;
+		const bool Scalar_IsValidInputH = (Scalar_InputPixelH >= 0) && (Scalar_InputPixelH < Hi);
+		
+		for (int ww = 0; ww < Ww; ++ww)
+		{
+			const int InputPixelW = Scalar_GroupInputPixelKernelTopLeftW + ww + OutputPixelOffset * StrideW;
+			const bool IsValidInput = Scalar_IsValidInputH && (InputPixelW >= 0) && (InputPixelW < Wi);
+			int InputPixelOffset = Scalar_InputPixelH * Wi + InputPixelW;
+			
+			InputPixelOffset = clamp(InputPixelOffset, 0, Hi * Wi);
+			
+			for (int ci = 0; ci < Ci; ci += LDS_IC)
+			{
+				///— DDR → LDS
+				// We need to load 1024 floats
+				// 512 weights (16 inputs channel, 32 output channels) and
+				// 512 inputs (16 inputs channel, 32 inputs pixels)
+				// We have 64 threads thus each will read 8 inputs and 8 weights.
+
+				const int ReadOffsetI = Scalar_BatchInputOffset + (ci + InputChannelOffset) * Hi * Wi + InputPixelOffset;
+				#if WEIGHTS_TRANSPOSED
+					const int InputChannelStride = Cw;
+					const int ReadOffsetW = Scalar_KernelIdx + (ci + InputChannelOffset) * InputChannelStride + (Scalar_GroupOutputChannelOffset + OutputChannelOffset);
+				#else
+					const int InputChannelStride = Hw * Ww;
+					const int ReadOffsetW = (Scalar_GroupOutputChannelOffset + OutputChannelOffset) * Ci * InputChannelStride + (ci + InputChannelOffset) * InputChannelStride + Scalar_KernelIdx;
+				#endif
+				
+				UNROLL
+				for (i = 0; i < (LDS_IC / 2); ++i)
+				{
+					WORK_TYPE ValueI = BUFFER_TO_WORK_TYPE(Input[Hi * Wi * i + ReadOffsetI]);
+					ValueI = IsValidInput ? ValueI : 0.0f;
+					WriteLDSInputs(InputChannelOffset + i, OutputPixelOffset, ValueI);
+
+					WORK_TYPE ValueW = BUFFER_TO_WORK_TYPE(Weight[InputChannelStride * i + ReadOffsetW]);
+					WriteLDSWeights(InputChannelOffset + i, OutputChannelOffset, ValueW);
+				}
+				
+				GroupMemoryBarrierWithGroupSync();
+				
+				/// LDS to register + inner loop
+				// Loop on cached input channels
+				#define INNERLOOPUNROLLCOUNT 4
+				for (int cachedCiBase = 0; cachedCiBase < LDS_IC; cachedCiBase += INNERLOOPUNROLLCOUNT)
+				{
+					WORK_TYPE RegWeights[INNERLOOPUNROLLCOUNT][4];
+					WORK_TYPE RegInputs[INNERLOOPUNROLLCOUNT][4];
+					
+					UNROLL
+					for (int unrolledIdx = 0; unrolledIdx < INNERLOOPUNROLLCOUNT; ++unrolledIdx)
+					{
+						// 8 load from LDS
+						UNROLL
+						for (i = 0; i < 4; ++i)
+						{
+							RegWeights[unrolledIdx][i] = ReadLDSWeights(cachedCiBase + unrolledIdx, GroupThreadOutputChannelOffset + i);
+							RegInputs[unrolledIdx][i] = ReadLDSInputs(cachedCiBase + unrolledIdx, GroupThreadOutputPixelOffset + i);
+						}
+
+						// Inner loop (16 mads)
+						UNROLL
+						for (i = 0; i < 4; ++i)
+						{
+							Values[i][0] += RegWeights[unrolledIdx][i] * RegInputs[unrolledIdx][0];
+							Values[i][1] += RegWeights[unrolledIdx][i] * RegInputs[unrolledIdx][1];
+							Values[i][2] += RegWeights[unrolledIdx][i] * RegInputs[unrolledIdx][2];
+							Values[i][3] += RegWeights[unrolledIdx][i] * RegInputs[unrolledIdx][3];
+						}
+					}
+				}
+				
+				GroupMemoryBarrierWithGroupSync();
+			}
+			#if WEIGHTS_TRANSPOSED
+				Scalar_KernelIdx += Ci * Cw;
+			#else
+			++Scalar_KernelIdx;
+			#endif
+		}
+	}
+	
+	/// Write results back to DDR (via LDS)
+	UNROLL
+	for (i = 0; i < 4; ++i)
+	{
+		const int LDSWriteOffsetY = 4 * GroupThreadID.y + i;
+		const int LDSWriteOffsetX = (LDS_PADX + 4) * GroupThreadID.x;
+		LDSBuffer[LDSWriteOffsetY * LDS_X_PADDEDSIZE + LDSWriteOffsetX + 0] = Values[i][0];
+		LDSBuffer[LDSWriteOffsetY * LDS_X_PADDEDSIZE + LDSWriteOffsetX + 1] = Values[i][1];
+		LDSBuffer[LDSWriteOffsetY * LDS_X_PADDEDSIZE + LDSWriteOffsetX + 2] = Values[i][2];
+		LDSBuffer[LDSWriteOffsetY * LDS_X_PADDEDSIZE + LDSWriteOffsetX + 3] = Values[i][3];
+	}
+
+	GroupMemoryBarrierWithGroupSync();
+
+	const int WriteOutChannelOffset = GroupThreadIdx >= 32 ? LDS_IC : 0;
+	const int WriteOutPixelOffset = GroupThreadIdx % 32;
+	const int LDSPixelOffset = (4 + LDS_PADX) * (WriteOutPixelOffset / 4) + WriteOutPixelOffset % 4;
+	
+	UNROLL
+	for (i = 0; i < LDS_IC; ++i)
+	{
+		const int WriteOffset = Scalar_BatchOutputOffset + (Scalar_GroupOutputChannelOffset + i + WriteOutChannelOffset) * Ho * Wo + Scalar_GroupOutputPixelOffset + WriteOutPixelOffset;
+		Output[WriteOffset] = LDSBuffer[(i + WriteOutChannelOffset) * LDS_X_PADDEDSIZE + LDSPixelOffset];
+	}
+
+}
\ No newline at end of file
diff --git a/Engine/Plugins/Experimental/NNERuntimeRDG/Shaders/Private/NNEHlslShaders/NNEHlslShadersConvTranspose.usf b/Engine/Plugins/Experimental/NNERuntimeRDG/Shaders/Private/NNEHlslShaders/NNEHlslShadersConvTranspose.usf
new file mode 100644
--- /dev/null
+++ b/Engine/Plugins/Experimental/NNERuntimeRDG/Shaders/Private/NNEHlslShaders/NNEHlslShadersConvTranspose.usf
@@ -0,0 +1,334 @@
+// Copyright Epic Games, Inc. All Rights Reserved.
+
+#include "/Engine/Public/Platform.ush"
+
+#if DATA_TYPE == 0
+
+#define WORK_TYPE float
+#define BUFFER_TYPE float
+#define BUFFER_TO_WORK_TYPE(x) x
+#define WORK_TO_BUFFER_TYPE(x) x
+
+#elif DATA_TYPE == 1
+
+#define WORK_TYPE int
+#define BUFFER_TYPE int
+#define BUFFER_TO_WORK_TYPE(x) x
+#define WORK_TO_BUFFER_TYPE(x) x
+
+#elif DATA_TYPE == 2
+
+#define WORK_TYPE uint
+#define BUFFER_TYPE uint
+#define BUFFER_TO_WORK_TYPE(x) x
+#define WORK_TO_BUFFER_TYPE(x) x
+
+#elif DATA_TYPE == 3
+
+#define WORK_TYPE float
+#define BUFFER_TYPE min16uint
+#define BUFFER_TO_WORK_TYPE(x) f16tof32(x)
+#define WORK_TO_BUFFER_TYPE(x) f32tof16(x)
+
+#endif
+
+// Must correspond to EConvGroupSize defined in NNEHlslShadersConvCS.h
+#if GROUP_SIZE == 0
+#define NUM_GROUP_THREADS 128
+#elif GROUP_SIZE == 1
+#define NUM_GROUP_THREADS 256
+#elif GROUP_SIZE == 2
+#define NUM_GROUP_THREADS 512
+#endif
+
+Buffer<BUFFER_TYPE> X; // N x C x XD0 x ... x XDi
+Buffer<BUFFER_TYPE> W; // M x C/GROUPS x WD0 x ... x WDi
+RWBuffer<BUFFER_TYPE> Y;  // N x M x YD0 x ... x YDi
+Buffer<BUFFER_TYPE> B; // B0, ..., BM
+
+// x: Dilation0, ...., Dilationi
+// y: Strides0, ...., Stridesi
+// z: Pad0Begin, ...., PadiBegin
+int4 Dilation_Stride_XBlockStartOffset_DilationXBlockStride[NUM_STACK_DIMENSIONS];
+
+// x: Element j contains PROD(GDj+1,..,GDi) with the last element containing 1, with GDi being the number of thread groups contained by a dimension i
+// y: Number of threads in each dimension of a group
+// z: Element j contains PROD(GTDj+1,..,GTDi) with the last element containing 1, with GTDi being the number of threads contained by a dimension i inside a group
+int4 GroupStride_GroupShape_GroupThreadStride_StrideXBlockStride[NUM_STACK_DIMENSIONS];
+
+// x: YD0, ..., YDi
+// y: Element j contains PROD(YDj+1,..,YDi) with the last element containing 1
+// z: XD0, ..., XDi
+// w: Element j contains PROD(XDj+1,..,XDi) with the last element containing 1
+int4 YDimension_YMemoryStride_XDimension_XMemoryStride[NUM_STACK_DIMENSIONS];
+
+// x: Number of elements in each dimension of a X block
+// y: The strides of the X block to be loaded
+// z: WD0, ..., WDi
+int4 XBlockStartStride_XBlockStride_WDimension_WDimensionDilationXBlockStride[NUM_STACK_DIMENSIONS];
+
+// x: 1/GroupStride
+// y: 1/GroupThreadStride
+// z: 1/XBlockStride
+float4 OneDiv_GroupStride_GroupThreadStride_OneDivStride[NUM_STACK_DIMENSIONS];
+
+int NumWChannels; // WShape[0]
+int NumOutChannelsDivGroup; // WShape[1]
+
+int YBatchStride;
+int YOutputKernelStride;
+
+int XBatchStride; // The number of elements in each bacth of X
+int XChannelStride; // The number of elements in each channel of X
+
+int XBlockSize; // The total number of elements in each block to be loaded
+
+int NumChannelBatches; // The number of iteration needed to cover all channels that needs to be processed. Equals to ceil(NumWChannels/NumChannelsPerBatch)
+int NumChannelsPerBatch; // The number of complete kernel channels that can be loaded in one batch
+
+int WOutputKernelStride; // The total number of W elements per output kernel and thus PROD(C/GROUPS, WD0, .., WDi)
+int WChannelBatchSize; // The total number of W elements inside a channel batch and thus the number of channels per batch times PROD(WD0, .., WDi)
+int WChannelSize; // Number of elements in a single channel of W and thus PROD(WDi)
+
+float GroupsDivM; // The number of groups divided by the number of output kernels used to compute the start channel offset
+float OneDivGroup;
+
+groupshared WORK_TYPE SharedMemoryX[NUM_GROUP_THREADS<<NUM_READS_PER_THREAD_POW2];
+groupshared WORK_TYPE SharedMemoryW[NUM_GROUP_THREADS];
+
+// Must correspond to EConvAlgorithm defined in NNEHlslShadersConvCS.h
+#if ALGORITHM < 1
+
+[numthreads(NUM_GROUP_THREADS, 1, 1)]
+void ConvTranspose(in const int3 DispatchThreadID : SV_DispatchThreadID, in const int3 GroupID : SV_GroupID, in const int3 GroupThreadID : SV_GroupThreadID)
+{
+	// Group Z dimension and Y dimension are used for batch and output kernel
+    // This works as each group only works on the inner dimensions (num threads in y and z set to one, group count to N and M respectively)
+	int BatchIndex = GroupID.z;
+	int OutputKernelIndex = GroupID.y;
+
+	// Input channels and output kernels are divided into groups
+	// Thus the channel offset for this thread's output kernel can be computed as follows:
+	// The output kernel index is divided by the number of output kernels per group giving the group index
+	// This is multiplied by the number of channels in W to get the start index
+	int InputChannelStartIndex = (int)(GroupsDivM * (float)OutputKernelIndex);
+
+	// Compute the output memory index and dimension indices of this thread inside its group as well as the start dimension indices of the X block to load
+	int YIndex = BatchIndex * YBatchStride + OutputKernelIndex * YOutputKernelStride;
+	int GroupThreadDimensionIndex[NUM_STACK_DIMENSIONS];// This will be used later to lookup the shared W and X to compute the final result
+    int XBlockStartDimensionIndex[NUM_STACK_DIMENSIONS];
+    float ResidualXIndex[NUM_STACK_DIMENSIONS];
+	int GroupIndex = GroupID.x;
+	int GroupThreadIndex = GroupThreadID.x;
+	UNROLL
+	for (int i = 0; i < NUM_STACK_DIMENSIONS; i++)
+	{
+        // Divide by the number of groups contained by one index of this dimension to get the dimension index n being the nth group in this dimension
+		int GroupDimensionIndex = (int)(OneDiv_GroupStride_GroupThreadStride_OneDivStride[i].x *(float)GroupIndex);
+		GroupIndex -= GroupDimensionIndex * GroupStride_GroupShape_GroupThreadStride_StrideXBlockStride[i].x;
+
+        // Same here with number of threads contained inside the group giving the nth thread inside the group in this dimension
+		GroupThreadDimensionIndex[i] = OneDiv_GroupStride_GroupThreadStride_OneDivStride[i].y * (float)GroupThreadIndex;
+		GroupThreadIndex -= GroupThreadDimensionIndex[i] * (float)GroupStride_GroupShape_GroupThreadStride_StrideXBlockStride[i].z;
+
+        float FXBlockStartDimensionIndex = ((float)GroupDimensionIndex * (float)XBlockStartStride_XBlockStride_WDimension_WDimensionDilationXBlockStride[i].x + (float)Dilation_Stride_XBlockStartOffset_DilationXBlockStride[i].z) * OneDiv_GroupStride_GroupThreadStride_OneDivStride[i].w;
+        ResidualXIndex[i] = (float)(FXBlockStartDimensionIndex - int(FXBlockStartDimensionIndex)) * (float)Dilation_Stride_XBlockStartOffset_DilationXBlockStride[i].y;
+
+        // Compose the global index by the group index times group shape plus thread index inside the group
+		int ThreadDimensionIndex = GroupDimensionIndex * GroupStride_GroupShape_GroupThreadStride_StrideXBlockStride[i].y + (int)GroupThreadDimensionIndex[i];
+
+        // Subtract the padding from the index of the first element of the group to get the start index of the block needed by this group
+        XBlockStartDimensionIndex[i] = (int)FXBlockStartDimensionIndex;
+
+        // ThreadDimensionIndex is always positive but needs to be tested if it does not overflow
+		if (ThreadDimensionIndex >= YDimension_YMemoryStride_XDimension_XMemoryStride[i].x || YIndex < 0)
+		{
+			YIndex = -1;
+		}
+		else
+		{
+			YIndex += ThreadDimensionIndex * YDimension_YMemoryStride_XDimension_XMemoryStride[i].y;
+		}
+	}
+
+    // Compute the memory indices in global X which the current thread has to load as well as the indices in shared X where to write
+	int GlobalXIndex[0x1<<NUM_READS_PER_THREAD_POW2];
+	int SharedXIndex[0x1<<NUM_READS_PER_THREAD_POW2];
+	int ReadIdx = GroupThreadID.x;
+    UNROLL
+    for (int j = 0; j < (int)(0x1<<NUM_READS_PER_THREAD_POW2); j++)
+	{
+        // If the read index exceeded the number of elements in the X block further computation can be skipped
+        if (ReadIdx >= XBlockSize)
+        {
+            // Assign -1 for both to indicate that the element is not used and thus neither has to be read from global nor written to shared memory
+            GlobalXIndex[j] = -1;
+            SharedXIndex[j] = -1;
+            continue;
+        }
+
+        // Turn the flat group index multiple into thread dimension indices inside the volume to be loaded by this group
+        int TmpGlobalXIndex = BatchIndex * XBatchStride + InputChannelStartIndex * XChannelStride;
+        int TmpSharedXIndex = 0;
+        int TempReadIdx = ReadIdx;
+        UNROLL
+        for (int k = 0; k < NUM_STACK_DIMENSIONS; k++)
+        {
+            int DimensionIndex = (int)(OneDiv_GroupStride_GroupThreadStride_OneDivStride[k].z * (float)TempReadIdx);
+            int DimensionIndexOffset = DimensionIndex * XBlockStartStride_XBlockStride_WDimension_WDimensionDilationXBlockStride[k].y;
+		    TempReadIdx -= DimensionIndexOffset;
+
+            // Update TmpGlobalXIndex only if positive (it is already outside the volume if negative)
+            if (TmpGlobalXIndex >= 0)
+            {
+                // Check if the global X index is in range
+                int XIndex = XBlockStartDimensionIndex[k] + DimensionIndex;
+                if (XIndex < 0 || XIndex >= YDimension_YMemoryStride_XDimension_XMemoryStride[k].z)
+                {
+                    TmpGlobalXIndex = -1;
+                }
+                else
+                {
+                    TmpGlobalXIndex += XIndex * YDimension_YMemoryStride_XDimension_XMemoryStride[k].w;
+                }
+            }
+
+            // Update TmpSharedXIndex in any case as a zero has to be written
+            TmpSharedXIndex += DimensionIndexOffset;
+        }
+
+        // Assign the result, which will be -1 for the global and positive for the shared index if the element lies in the padding area
+        // This indicates that the global element should not be loaded but a zero has to be written to shared memory
+        GlobalXIndex[j] = TmpGlobalXIndex;
+        SharedXIndex[j] = TmpSharedXIndex;
+
+        // Increment the read idx by the number of threads inside the group to cover all elements
+        ReadIdx += NUM_GROUP_THREADS;
+    }
+ 
+    // Write the zeros once to avoid rewriting the same value later
+    UNROLL
+    for (int l = 0; l < (int)(0x1<<NUM_READS_PER_THREAD_POW2); l++)
+	{
+        // If the global index is out of bound (padding area) but the shared index is valid (inside the X block) the shared element needs initialization to zero
+        if (GlobalXIndex[l] < 0 && SharedXIndex[l] >= 0)
+        {
+            SharedMemoryX[SharedXIndex[l]] = 0;
+        }
+    }
+
+    // Iterate over all channel batches
+    int WChannelBatchOffset = 0;
+    int WChannelIndex = 0;
+#if HAS_B == 0
+    WORK_TYPE Result = 0;
+#else
+    WORK_TYPE Result = BUFFER_TO_WORK_TYPE(B[OutputKernelIndex]);
+#endif
+    // Each channel batch iterates over all consecutive elements of the loaded kernel channels
+    // This results in bank conflict free lookups, as all threads of the same warp read exactly the same element which is then broadcasted
+    for (int m = 0; m < NumChannelBatches; m++)
+	{
+        // Load all channels of W belonging to this batch
+        int SharedWIndex = 0;
+        WChannelBatchOffset = m * WChannelBatchSize;
+
+        // Iterate over all channels in this batch
+        for (int n = 0; n < NumChannelsPerBatch; n++)
+	    {
+            int WChannelOutputKernelOffset = WChannelSize * OutputKernelIndex + NumOutChannelsDivGroup * int((float)n * OneDivGroup);
+            int WIndex = WChannelBatchOffset + GroupThreadID.x;
+            if (GroupThreadID.x < WChannelBatchSize && WIndex < WOutputKernelStride)
+            {  
+                SharedMemoryW[GroupThreadID.x] = BUFFER_TO_WORK_TYPE(W[WChannelOutputKernelOffset + WIndex]);
+
+            }
+            if (WChannelOutputKernelOffset < WChannelSize * NumOutChannelsDivGroup * n ||
+                WChannelOutputKernelOffset >= NumOutChannelsDivGroup * WChannelSize * (n + 1)) {
+                SharedMemoryW[GroupThreadID.x] = 0;
+            }
+
+            // Check if the channel index is valid as channel batches may not be aligned with the effective number of channels
+            if (WChannelIndex >= NumWChannels)
+            {
+                break;
+            }
+            WChannelIndex++;
+
+            // Load the shared X
+            UNROLL
+            for (int o = 0; o < (int)(0x1<<NUM_READS_PER_THREAD_POW2); o++)
+	        {
+                // If the global index is negative, a zero has already been written, if shared index is negative there is no need to load anything
+                if (GlobalXIndex[o] >= 0 && SharedXIndex[o] >= 0)
+                {
+                    // Load the global data into shared memory
+                    SharedMemoryX[SharedXIndex[o]] = BUFFER_TO_WORK_TYPE(X[GlobalXIndex[o]]);
+
+                    // Update the indices for the next iteration by adding the channel stride
+                    GlobalXIndex[o] += XChannelStride;
+                }
+            }
+
+            // Sync the group before doing any computation
+            GroupMemoryBarrierWithGroupSync();
+
+            // Initialize a dimension index array used to iterate through the kernel dimensions
+            // Also initialize the start index in shared memory based on the result position
+            int WDimensionIndex[NUM_STACK_DIMENSIONS];
+            float SharedXIndex = 0;
+
+            UNROLL
+            for (int p = 0; p < NUM_STACK_DIMENSIONS; p++)
+            {
+                WDimensionIndex[p] = 0;
+
+                // The index of the first element (e.g. upper left) required by each output is offsetted by the kernel stride and then multiplied by the stride in memory
+                SharedXIndex += (GroupThreadDimensionIndex[p] + ResidualXIndex[p]) * OneDiv_GroupStride_GroupThreadStride_OneDivStride[p].w * (float)GroupStride_GroupShape_GroupThreadStride_StrideXBlockStride[p].w;
+                }
+
+            // Compute and accumulate the result
+            for (int q = 0; q < WChannelSize; q++)
+            {
+                if (SharedXIndex - int(SharedXIndex+0.0001) < 0.001 && SharedXIndex > -0.0001){
+                    Result += SharedMemoryX[SharedXIndex+0.0001] * SharedMemoryW[WChannelSize - 1 - q];
+                }
+
+                // Advance the inner-most dimension
+                SharedWIndex++;
+                SharedXIndex += Dilation_Stride_XBlockStartOffset_DilationXBlockStride[NUM_STACK_DIMENSIONS-1].x * OneDiv_GroupStride_GroupThreadStride_OneDivStride[NUM_STACK_DIMENSIONS-1].w;
+
+                WDimensionIndex[NUM_STACK_DIMENSIONS-1]++;
+
+                // Advance the other dimensions if necessary
+                UNROLL
+                for (int r = NUM_STACK_DIMENSIONS-1; r > 0; r--)
+                {
+                    if(WDimensionIndex[r] >= XBlockStartStride_XBlockStride_WDimension_WDimensionDilationXBlockStride[r].z)
+                    {
+                        // Adjust the SharedXIndex by resetting the current dimension and increase the next outer one
+                        SharedXIndex -= XBlockStartStride_XBlockStride_WDimension_WDimensionDilationXBlockStride[r].w * OneDiv_GroupStride_GroupThreadStride_OneDivStride[r].w;
+                        SharedXIndex += Dilation_Stride_XBlockStartOffset_DilationXBlockStride[r-1].w * OneDiv_GroupStride_GroupThreadStride_OneDivStride[r-1].w;
+                        WDimensionIndex[r] = 0;
+                        WDimensionIndex[r-1]++;
+                    }
+                    else
+                    {
+                        break;
+                    }
+                }
+            }
+
+            // Sync the group before continuing any load
+            GroupMemoryBarrierWithGroupSync();
+        }
+    }
+
+    // Write the final result
+    if (YIndex >= 0)
+    {
+        Y[YIndex] = WORK_TO_BUFFER_TYPE(Result);
+    }
+}
+
+#endif
diff --git a/Engine/Plugins/Experimental/NNERuntimeRDG/Shaders/Private/NNEHlslShaders/NNEHlslShadersElementWiseBinary.usf b/Engine/Plugins/Experimental/NNERuntimeRDG/Shaders/Private/NNEHlslShaders/NNEHlslShadersElementWiseBinary.usf
new file mode 100644
--- /dev/null
+++ b/Engine/Plugins/Experimental/NNERuntimeRDG/Shaders/Private/NNEHlslShaders/NNEHlslShadersElementWiseBinary.usf
@@ -0,0 +1,61 @@
+// Copyright Epic Games, Inc. All Rights Reserved.
+
+#include "/Engine/Public/Platform.ush"
+#include "NNEHlslShadersBroadcastHelper.ush"
+
+Buffer<float> LHSInput;
+Buffer<float> RHSInput;
+RWBuffer<float> Output;
+uint4 TensorInfo[NUM_DIMENSIONS];
+uint Num;
+uint ThreadCountX;
+
+// See [REDACTED_PROFILE_URL]/onnx/blob/main/docs/Operators.md#prelu
+float prelu(float x, float slope)
+{
+	return (x < 0.0f) ? (slope * x) : x;
+}
+
+float safe_pow(float base, float xp)
+{
+	if (base >= 0.0f)
+		return pow(base, xp);
+	
+	//negative base + not integer exponent => nan
+	if (round(xp) != xp)
+	{
+		
+		const float nan = 0. / 0.;
+		return nan;
+	}
+
+	//negative base + integer exponent
+	int isExpOdd = (xp % 2);
+	float sign = isExpOdd ? -1.0f : 1.0f;
+	return sign * pow(-base, xp);
+}
+
+[numthreads(THREADGROUP_SIZE_X, 1, 1)]
+void ElementWiseBinary(in const uint3 DispatchThreadID : SV_DispatchThreadID)
+{
+	const uint Index = DispatchThreadID.y * ThreadCountX + DispatchThreadID.x;
+	
+	if (Index < Num)
+	{
+		uint LHSIndex = 0;
+		uint RHSIndex = 0;
+		uint Offset = Index;
+		for (uint dim = 0; dim < NUM_DIMENSIONS; ++dim)
+		{
+			uint Q, R;
+			DivMod(Offset, TensorInfo[dim][2], Q, R);
+			LHSIndex += TensorInfo[dim][0] * Q;
+			RHSIndex += TensorInfo[dim][1] * Q;
+			Offset = R;
+		}
+
+		float X = LHSInput[LHSIndex];
+		float Y = RHSInput[RHSIndex];
+		Output[Index] = ELEMENTWISE_OP(X, Y);
+	}
+}
diff --git a/Engine/Plugins/Experimental/NNERuntimeRDG/Shaders/Private/NNEHlslShaders/NNEHlslShadersElementWiseUnary.usf b/Engine/Plugins/Experimental/NNERuntimeRDG/Shaders/Private/NNEHlslShaders/NNEHlslShadersElementWiseUnary.usf
new file mode 100644
--- /dev/null
+++ b/Engine/Plugins/Experimental/NNERuntimeRDG/Shaders/Private/NNEHlslShaders/NNEHlslShadersElementWiseUnary.usf
@@ -0,0 +1,147 @@
+// Copyright Epic Games, Inc. All Rights Reserved.
+
+#include "/Engine/Public/Platform.ush"
+
+Buffer<float> Input;
+RWBuffer<float> Output;
+uint Num;
+uint ThreadCountX;
+float Alpha;
+Buffer<float> AlphaTensor;
+float Beta;
+Buffer<float> BetaTensor;
+float Gamma;
+
+float GetAlpha()
+{
+#if ALPHA_ON_GPU == 0
+	return Alpha;
+#else
+	return AlphaTensor[0];
+#endif
+}
+
+float GetBeta()
+{
+#if BETA_ON_GPU == 0
+	return Beta;
+#else
+	return BetaTensor[0];
+#endif
+}
+
+//see [REDACTED_PROFILE_URL]/onnx/blob/main/docs/Operators.md#Relu
+float relu(float x)
+{
+	return max(0.0f, x);
+}
+
+//see [REDACTED_PROFILE_URL]/onnx/blob/main/docs/Operators.md#Selu
+float selu(float x)
+{
+	float yNegOrZero = Gamma * (GetAlpha() * exp(x) - GetAlpha());
+	float yPos = Gamma * x;
+	return (x > 0.0f) ? yPos : yNegOrZero;
+}
+
+//see [REDACTED_PROFILE_URL]/onnx/blob/main/docs/Operators.md#sigmoid
+float sigmoid(float x)
+{
+	return 1.0f / (1.0f + exp(-x));
+}
+
+//see [REDACTED_PROFILE_URL]/onnx/blob/main/docs/Operators.md#softplus
+float softplus(float x)
+{
+	return log(exp(x) + 1.0f);
+}
+
+//see [REDACTED_PROFILE_URL]/onnx/blob/main/docs/Operators.md#softsign
+float softsign(float x)
+{
+	return x/(1.0f+abs(x));
+}
+
+//see [REDACTED_PROFILE_URL]/onnx/blob/main/docs/Operators.md#acosh
+float acosh(float x)
+{
+	//https://mathworld.wolfram.com/InverseHyperbolicCosine.html
+	const static float NaN = 0.0f / 0.0f;
+	float yAboveOne = log( x + sqrt( x + 1.0f ) + sqrt( x - 1.0f ) );
+	return (x == 1.0f) ? 0.0f : (x >= 1.0f) ? yAboveOne : NaN;
+}
+
+//see [REDACTED_PROFILE_URL]/onnx/blob/main/docs/Operators.md#asinh
+float asinh(float x)
+{
+	//https://mathworld.wolfram.com/InverseHyperbolicSine.html
+	return log( x + sqrt( 1 + ( x * x ) ) );
+}
+
+//see [REDACTED_PROFILE_URL]/onnx/blob/main/docs/Operators.md#atanh
+float atanh(float x)
+{
+	//https://mathworld.wolfram.com/InverseHyperbolicTangent.html
+	return 0.5f * ( log( 1 + x) - log ( 1 - x) );
+}
+
+//see [REDACTED_PROFILE_URL]/onnx/blob/main/docs/Operators.md#elu
+float elu(float x)
+{
+    float yNeg = GetAlpha() * (exp(x) - 1.0f);
+	float yPosOrZero = x;
+	return (x >= 0.0f) ? yPosOrZero : yNeg;
+}
+
+//see [REDACTED_PROFILE_URL]/onnx/blob/main/docs/Operators.md#erf
+float erf(float x)
+{
+	//https://aapt.scitation.org/doi/abs/10.1119/1.15018?journalCode=ajp
+	float a = 167.0f/148.0f;
+	float b = 11.0f/109.0f;
+	float x3 = x*x*x;
+	return tanh( a*x + b*x3 );
+}
+
+float hardSigmoid(float x, float alpha, float beta)
+{
+	return max(0.0f, min(1.0f, alpha * x + beta));
+}
+
+//see [REDACTED_PROFILE_URL]/onnx/blob/main/docs/Operators.md#hardSigmoid
+float hardSigmoid(float x)
+{
+	return hardSigmoid(x, GetAlpha(), GetBeta());
+}
+
+//see [REDACTED_PROFILE_URL]/onnx/blob/main/docs/Operators.md#hardSwish
+float hardSwish(float x)
+{
+    return x * hardSigmoid(x, 1.0f/6.0f, 0.5f);
+}
+
+//see [REDACTED_PROFILE_URL]/onnx/blob/main/docs/Operators.md#leakyRelu
+float leakyRelu(float x)
+{
+    return (x >= 0.0f) ? x : GetAlpha() * x;
+}
+
+
+//see [REDACTED_PROFILE_URL]/onnx/blob/main/docs/Operators.md#clip
+float clipOp(float x)
+{
+	float min = GetAlpha();
+	float max = GetBeta();
+	return clamp(x, GetAlpha(), GetBeta());
+}
+
+
+[numthreads(THREADGROUP_SIZE_X, 1, 1)]
+void ElementWiseUnary(in const uint3 DispatchThreadID : SV_DispatchThreadID)
+{
+	const uint Index = DispatchThreadID.y * ThreadCountX + DispatchThreadID.x;
+	if (Index < Num)
+	{
+		Output[Index] = ELEMENTWISE_OP(Input[Index]);
+	}
+}
diff --git a/Engine/Plugins/Experimental/NNERuntimeRDG/Shaders/Private/NNEHlslShaders/NNEHlslShadersElementWiseVariadic.usf b/Engine/Plugins/Experimental/NNERuntimeRDG/Shaders/Private/NNEHlslShaders/NNEHlslShadersElementWiseVariadic.usf
new file mode 100644
--- /dev/null
+++ b/Engine/Plugins/Experimental/NNERuntimeRDG/Shaders/Private/NNEHlslShaders/NNEHlslShadersElementWiseVariadic.usf
@@ -0,0 +1,65 @@
+// Copyright Epic Games, Inc. All Rights Reserved.
+
+#include "/Engine/Public/Platform.ush"
+#include "NNEHlslShadersBroadcastHelper.ush"
+
+#define BUFFER_TYPE float
+
+Buffer<BUFFER_TYPE> Input0;
+Buffer<BUFFER_TYPE> Input1;
+Buffer<BUFFER_TYPE> Input2;
+Buffer<BUFFER_TYPE> Input3;
+RWBuffer<BUFFER_TYPE> Output;
+uint4 InputTensorInfo[NUM_DIMENSIONS];
+uint4 OutputTensorInfo[NUM_DIMENSIONS];
+uint Num;
+uint ThreadCountX;
+float Scale;
+
+[numthreads(THREADGROUP_SIZE_X, 1, 1)]
+void ElementWiseVariadic(in const uint3 DispatchThreadID : SV_DispatchThreadID)
+{
+	const uint Index = DispatchThreadID.y * ThreadCountX + DispatchThreadID.x;
+	
+	if (Index < Num)
+	{
+		uint Input0Index = 0;
+		uint Input1Index = 0;
+		uint Input2Index = 0;
+		uint Input3Index = 0;
+		uint Offset = Index;
+		for (uint dim = 0; dim < NUM_DIMENSIONS; ++dim)
+		{
+			uint Q, R;
+			DivMod(Offset, OutputTensorInfo[dim][0], Q, R);
+			Input0Index += InputTensorInfo[dim][0] * Q;
+			Input1Index += InputTensorInfo[dim][1] * Q;
+			Input2Index += InputTensorInfo[dim][2] * Q;
+			Input3Index += InputTensorInfo[dim][3] * Q;
+			Offset = R;
+		}
+
+		BUFFER_TYPE Result = Input0[Input0Index];
+
+		 // Must correspond to EVariadicNumInput defined in NNEElementWiseVariadicCS.h
+		#if NUMINPUT >= 2
+			Result  = ELEMENTWISE_OP(Result, Input1[Input1Index]);
+		#endif
+		#if NUMINPUT >= 3
+			Result  = ELEMENTWISE_OP(Result, Input2[Input2Index]);
+		#endif
+		#if NUMINPUT >= 4
+			Result  = ELEMENTWISE_OP(Result, Input3[Input3Index]);
+		#endif
+
+		#if OUTPUTASINPUT
+			Result  = ELEMENTWISE_OP(Result, Output[Index]);
+		#endif
+
+		#if APPLYSCALE
+			Result  = (BUFFER_TYPE)(Result * Scale);
+		#endif
+
+		Output[Index] = Result;
+	}
+}
diff --git a/Engine/Plugins/Experimental/NNERuntimeRDG/Shaders/Private/NNEHlslShaders/NNEHlslShadersGather.usf b/Engine/Plugins/Experimental/NNERuntimeRDG/Shaders/Private/NNEHlslShaders/NNEHlslShadersGather.usf
new file mode 100644
--- /dev/null
+++ b/Engine/Plugins/Experimental/NNERuntimeRDG/Shaders/Private/NNEHlslShaders/NNEHlslShadersGather.usf
@@ -0,0 +1,107 @@
+// Copyright Epic Games, Inc. All Rights Reserved.
+
+#include "/Engine/Public/Platform.ush"
+
+#if DATA_TYPE == 0
+#define DATA_BUFFER_TYPE float
+#else
+#define DATA_BUFFER_TYPE float
+#endif
+
+#if INDICES_TYPE == 0
+#define INDICES_WORK_TYPE int
+#define INDICES_BUFFER_TYPE int
+#define INDICES_BUFFER_TO_WORK_TYPE(x) x
+#else
+#define INDICES_WORK_TYPE int
+#define INDICES_BUFFER_TYPE int
+#define INDICES_BUFFER_TO_WORK_TYPE(x) x
+#endif
+
+int Axis;
+int OutputSize;
+int NumDataDimensions;
+int NumIndicesDimensions;
+int4 DataStride_IndicesStride_OutputStride_DataSizes[MAX_NUM_DIMENSIONS];
+float4 OneDivDataStride_OneDivIndicesStride_OneDivOutputStride[MAX_NUM_DIMENSIONS];
+
+
+Buffer<DATA_BUFFER_TYPE> Data;
+Buffer<INDICES_BUFFER_TYPE> Indices;
+RWBuffer<DATA_BUFFER_TYPE> Output;
+
+uint NormalizeIdx(int Idx, uint DimSize)
+{
+	return Idx < 0 ? (uint) (Idx + DimSize) : (uint) Idx;
+}
+
+// int64_t is only available from SM 6.0, therefore here we need to manually cast to int32.
+int CastInt64ToInt32(uint LSW, int MSW)
+{
+	return (MSW >> 31) == 1 ? - (int) LSW : (int) LSW;
+}
+
+[numthreads(NUM_GROUP_THREADS, 1, 1)]
+void Gather(in const uint3 DispatchThreadID : SV_DispatchThreadID)
+{
+	if (DispatchThreadID.x >= OutputSize)
+	{
+		return;
+	}
+
+	// Compute the output index per dimenion
+	INDICES_WORK_TYPE OutputIndices[NUM_OUTPUT_DIMENSIONS];
+	INDICES_WORK_TYPE DispatchThreadIndex = (INDICES_WORK_TYPE)DispatchThreadID.x;
+	UNROLL
+	for (int i = 0; i < NUM_OUTPUT_DIMENSIONS; i++)
+	{
+		OutputIndices[i] = (INDICES_WORK_TYPE)(OneDivDataStride_OneDivIndicesStride_OneDivOutputStride[i].z * (float)DispatchThreadIndex);
+		DispatchThreadIndex -= OutputIndices[i] * (INDICES_WORK_TYPE)DataStride_IndicesStride_OutputStride_DataSizes[i].z;
+	}
+
+	// Split the output indices into data and indices indices 
+	INDICES_WORK_TYPE DataIndices[NUM_OUTPUT_DIMENSIONS];
+	INDICES_WORK_TYPE IndicesIndices[NUM_OUTPUT_DIMENSIONS];
+	for (int i = 0; i < Axis; i++)
+	{
+		DataIndices[i] = OutputIndices[i];
+	}
+	for (int i = 0; i < NumIndicesDimensions; i++)
+	{
+		IndicesIndices[i] = OutputIndices[Axis + i];
+	}
+	for (int i = Axis + 1; i < NumDataDimensions; i++)
+	{
+		DataIndices[i] = OutputIndices[NumIndicesDimensions + i - 1];
+	}
+
+	// Compute the flattened indices index
+	INDICES_WORK_TYPE IndicesIndex = 0;
+	for (int i = 0; i < NumIndicesDimensions; i++)
+	{
+		IndicesIndex += IndicesIndices[i] * DataStride_IndicesStride_OutputStride_DataSizes[i].y;
+	}
+
+	// Look up and assign the missing data index
+	#if HAS_64BIT_INDICES == 0
+		DataIndices[Axis] = NormalizeIdx(
+			INDICES_BUFFER_TO_WORK_TYPE(Indices[IndicesIndex]), 
+			DataStride_IndicesStride_OutputStride_DataSizes[Axis].w
+		);
+	#else
+		DataIndices[Axis] = NormalizeIdx(
+			INDICES_BUFFER_TO_WORK_TYPE(CastInt64ToInt32(Indices[IndicesIndex * 2], Indices[IndicesIndex * 2 + 1])),
+			DataStride_IndicesStride_OutputStride_DataSizes[Axis].w
+		);
+	#endif
+
+	// Compute the flattened data index
+	INDICES_WORK_TYPE DataIndex = 0;
+	for (int i = 0; i < NumDataDimensions; i++)
+	{
+		DataIndex += DataIndices[i] * DataStride_IndicesStride_OutputStride_DataSizes[i].x;
+	}
+
+	// Write the result
+	Output[DispatchThreadID.x] = Data[DataIndex];
+}
\ No newline at end of file
diff --git a/Engine/Plugins/Experimental/NNERuntimeRDG/Shaders/Private/NNEHlslShaders/NNEHlslShadersGemm.usf b/Engine/Plugins/Experimental/NNERuntimeRDG/Shaders/Private/NNEHlslShaders/NNEHlslShadersGemm.usf
new file mode 100644
--- /dev/null
+++ b/Engine/Plugins/Experimental/NNERuntimeRDG/Shaders/Private/NNEHlslShaders/NNEHlslShadersGemm.usf
@@ -0,0 +1,447 @@
+// Copyright Epic Games, Inc. All Rights Reserved.
+
+#include "/Engine/Public/Platform.ush"
+
+#define WORK_TYPE float
+#define BUFFER_TYPE float
+#define READ(x) x
+#define WRITE(x) x
+
+WORK_TYPE Alpha;
+WORK_TYPE Beta;
+int TransA;
+int TransB;
+uint M;
+uint N;
+uint K;
+uint MxK;
+uint KxN;
+uint MxN;
+uint CWidth;
+uint CHeight;
+
+WORK_TYPE CScalar;
+
+Buffer<BUFFER_TYPE> A;
+Buffer<BUFFER_TYPE> B;
+Buffer<BUFFER_TYPE> C;
+RWBuffer<BUFFER_TYPE> Y;
+
+// x: StackShapeA[0], ..., StackShapeA[NUM_STACK_DIMENSIONS - 1]
+// y: StackShapeB[0], ..., StackShapeB[NUM_STACK_DIMENSIONS - 1]
+// z: StackStrideA[0], ..., StackStrideA[NUM_STACK_DIMENSIONS - 1]
+// w: StackStrideB[0], ..., StackStrideB[NUM_STACK_DIMENSIONS - 1]
+uint4 StackShapeA_StackShapeB_StackStrideA_StackStrideB[MAX_NUM_STACK_DIMENSIONS];
+
+uint3 GetMatrixStackOffsets(const uint3 GroupID)
+{
+// Must correspond to FGemmNumDimensions defined in NNEHlslShadersGemmCS.h
+#if NUM_STACK_DIMENSIONS == 0
+	return uint3(0, 0, 0);
+#else
+	uint stackCoordA[NUM_STACK_DIMENSIONS];
+	uint stackCoordB[NUM_STACK_DIMENSIONS];
+
+	uint stackStrideOutput[NUM_STACK_DIMENSIONS];
+	stackStrideOutput[NUM_STACK_DIMENSIONS - 1] = 1;
+	for (int i = NUM_STACK_DIMENSIONS - 2; i >= 0; i--)
+	{
+		uint maxDim = max(StackShapeA_StackShapeB_StackStrideA_StackStrideB[i + 1][0], StackShapeA_StackShapeB_StackStrideA_StackStrideB[i + 1][1]);
+		stackStrideOutput[i] = maxDim * stackStrideOutput[i + 1];
+	}
+
+	uint idx = GroupID.z;
+	for (int i = 0; i < NUM_STACK_DIMENSIONS; i++)
+	{
+		uint coord = idx / stackStrideOutput[i];
+		stackCoordA[i] = coord % StackShapeA_StackShapeB_StackStrideA_StackStrideB[i][0];
+		stackCoordB[i] = coord % StackShapeA_StackShapeB_StackStrideA_StackStrideB[i][1];
+
+		idx = idx % stackStrideOutput[i];
+	}
+
+	uint stackOffsetA = 0;
+	uint stackOffsetB = 0;
+	for (int i = 0; i < NUM_STACK_DIMENSIONS; i++)
+	{
+		stackOffsetA += stackCoordA[i] * StackShapeA_StackShapeB_StackStrideA_StackStrideB[i][2];
+		stackOffsetB += stackCoordB[i] * StackShapeA_StackShapeB_StackStrideA_StackStrideB[i][3];
+	}
+
+	return uint3(
+		stackOffsetA	* MxK,
+		stackOffsetB	* KxN,
+		GroupID.z		* MxN);
+#endif
+}
+
+WORK_TYPE GetBetaTimesC(const uint3 DispatchThreadID)
+{
+// Must correspond to EGemmCScalar defined in NNEHlslShadersGemmCS.h
+#if C_SCALAR == 0
+	return Beta * READ(C[(DispatchThreadID.y % CHeight) * N + (DispatchThreadID.x % CWidth)]);
+#elif C_SCALAR == 1
+	return Beta * CScalar;
+#else
+	return 0;
+#endif
+}
+
+// Must correspond to EGemmAlgorithm defined in NNEHlslShadersGemmCS.h
+#if ALGORITHM < 4
+
+#if ALGORITHM == 0
+[numthreads(8, 8, 1)]
+#elif ALGORITHM == 1
+[numthreads(16, 16, 1)]
+#elif ALGORITHM == 2
+[numthreads(32, 32, 1)]
+#elif ALGORITHM == 3
+[numthreads(256, 1, 1)]
+#endif
+void Gemm(in const uint3 DispatchThreadID : SV_DispatchThreadID, in const uint3 GroupID : SV_GroupID)
+{
+	if (DispatchThreadID.y >= M || DispatchThreadID.x >= N)
+	{
+		return;
+	}
+	
+	uint AIdx;
+	uint AIdxIncrement;
+	if (TransA == 0)
+	{
+		AIdx = DispatchThreadID.y * K; // Index of the first element in the row
+		AIdxIncrement = 1; // The loop goes over the whole row
+	}
+	else
+	{
+		AIdx = DispatchThreadID.y; // Index of the first element in the column
+		AIdxIncrement = M; // The loop goes over the column, thus skip a full row
+	}
+	
+	uint BIdx;
+	uint BIdxIncrement;
+	if (TransB == 0)
+	{
+		BIdx = DispatchThreadID.x; // Index of the first element in the column
+		BIdxIncrement = N; // The loop goes over the column, thus skip a full row
+	}
+	else
+	{
+		BIdx = DispatchThreadID.x * K; // Index of the first element in the row
+		BIdxIncrement = 1; // The loop goes over the whole row
+	}
+
+	const uint3 StackOffsets = GetMatrixStackOffsets(GroupID);
+	AIdx += StackOffsets.x;
+	BIdx += StackOffsets.y;
+	
+	WORK_TYPE Result = 0.0;
+	for (int k = 0; k < K; k++)
+	{
+		Result += READ(A[AIdx]) * READ(B[BIdx]);
+		AIdx += AIdxIncrement;
+		BIdx += BIdxIncrement;
+	}
+	
+	Y[StackOffsets.z + DispatchThreadID.y * N + DispatchThreadID.x] = WRITE(Alpha * Result + GetBetaTimesC(DispatchThreadID));
+}
+
+#elif ALGORITHM < 7
+
+#if ALGORITHM == 4
+#define GROUP_SIZE 8
+#elif ALGORITHM == 5
+#define GROUP_SIZE 16
+#elif ALGORITHM == 6
+#define GROUP_SIZE 32
+#endif
+
+groupshared WORK_TYPE SharedMemoryA[GROUP_SIZE * GROUP_SIZE];
+groupshared WORK_TYPE SharedMemoryB[GROUP_SIZE * GROUP_SIZE];
+
+[numthreads(GROUP_SIZE, GROUP_SIZE, 1)]
+void Gemm(in const uint3 DispatchThreadID : SV_DispatchThreadID, in const uint3 GroupThreadID : SV_GroupThreadID, in const uint3 GroupID : SV_GroupID)
+{
+	// Calculate the number of steps each group needs to make to cover the full length of the matrix
+	uint NumGroupSteps = (uint)ceil((float)K / (float)GROUP_SIZE);
+	
+	// Compute some variables required at multiple places
+	uint SharedMemoryLineStartIdx = GroupThreadID.y * GROUP_SIZE;
+	uint SharedMemoryIdx = SharedMemoryLineStartIdx + GroupThreadID.x;
+	
+	// Compute the lookup start indices and the increment which needs to be applied for each sub block
+	uint AIdx;
+	uint AIdxIncrement;
+	if (TransA == 0)
+	{
+		AIdx = DispatchThreadID.y * K + GroupThreadID.x;
+		AIdxIncrement = GROUP_SIZE;
+	}
+	else
+	{
+		AIdx = GroupThreadID.x * M + DispatchThreadID.y;
+		AIdxIncrement = GROUP_SIZE * M;
+	}
+	
+	uint BIdx;
+	uint BIdxIncrement;
+	if (TransB == 0)
+	{
+		BIdx = GroupThreadID.y * N + DispatchThreadID.x;
+		BIdxIncrement = GROUP_SIZE * N;
+	}
+	else
+	{
+		BIdx = DispatchThreadID.x * K + GroupThreadID.y;
+		BIdxIncrement = GROUP_SIZE;
+	}
+
+	const uint3 StackOffsets = GetMatrixStackOffsets(GroupID);
+	AIdx += StackOffsets.x;
+	BIdx += StackOffsets.y;
+	
+	// Multiply and accumulate the blocks along one dimension of the matrices
+	WORK_TYPE Result = 0.0;
+	uint GroupIdx = 0;
+	for (int GroupStep = 0; GroupStep < NumGroupSteps; GroupStep++) 
+	{
+		// Only sync and update inidces starting with the second iteration
+		if(GroupStep > 0)
+		{
+			// Increment the indices
+			GroupIdx += GROUP_SIZE;
+			AIdx += AIdxIncrement;
+			BIdx += BIdxIncrement;
+			
+			// Make sure computation is done in the whole group before overwriting shared memory
+			GroupMemoryBarrierWithGroupSync();
+		}
+		
+		// Load the shared memory after applying the transpose and set overflow elements to zero
+		WORK_TYPE Temp = 0;
+		if (DispatchThreadID.y < M && GroupIdx + GroupThreadID.x < K)
+		{
+			Temp = READ(A[AIdx]);
+		}
+		SharedMemoryA[SharedMemoryIdx] = Temp;
+
+		Temp = 0;
+		if (DispatchThreadID.x < N && GroupIdx + GroupThreadID.y < K)
+		{
+			Temp = READ(B[BIdx]);
+		}
+		SharedMemoryB[SharedMemoryIdx] = Temp;
+		
+		// Make sure all memory has been loaded before running the multiplication
+		GroupMemoryBarrierWithGroupSync();
+		
+		// Multiply the groups and accumulate the result
+		// This works also at matrix borders, as overflow elements are set to zero
+		uint SharedMemoryBIdx = GroupThreadID.x;
+		for (int kGroup = 0; kGroup < GROUP_SIZE; kGroup++)
+		{
+			Result += SharedMemoryA[SharedMemoryLineStartIdx + kGroup] * SharedMemoryB[SharedMemoryBIdx];
+			SharedMemoryBIdx += GROUP_SIZE;
+		}
+	}
+	
+	// Only now return out of bound threads as they were needed to load shared memory above
+	if (DispatchThreadID.y >= M || DispatchThreadID.x >= N)
+	{
+		return;
+	}
+	
+	Y[StackOffsets.z + DispatchThreadID.y * N + DispatchThreadID.x] = WRITE(Alpha * Result + GetBetaTimesC(DispatchThreadID));
+}
+
+#else
+
+#if ALGORITHM == 7
+#define GROUP_SIZE_X 		16
+#define GROUP_SIZE_Y 		1
+#define LOAD_PER_THREAD 	16
+#elif ALGORITHM == 8
+#define GROUP_SIZE_X 		16
+#define GROUP_SIZE_Y 		2
+#define LOAD_PER_THREAD 	8
+#elif ALGORITHM == 9
+#define GROUP_SIZE_X 		32
+#define GROUP_SIZE_Y 		1
+#define LOAD_PER_THREAD 	32
+#elif ALGORITHM == 10
+#define GROUP_SIZE_X 		32
+#define GROUP_SIZE_Y 		2
+#define LOAD_PER_THREAD 	16
+#elif ALGORITHM == 11
+#define GROUP_SIZE_X 		32
+#define GROUP_SIZE_Y 		4
+#define LOAD_PER_THREAD 	8
+#elif ALGORITHM == 12
+#define GROUP_SIZE_X 		64
+#define GROUP_SIZE_Y 		2
+#define LOAD_PER_THREAD 	32
+#elif ALGORITHM == 13
+#define GROUP_SIZE_X 		64
+#define GROUP_SIZE_Y 		4
+#define LOAD_PER_THREAD 	16
+#elif ALGORITHM == 14
+#define GROUP_SIZE_X 		64
+#define GROUP_SIZE_Y 		8
+#define LOAD_PER_THREAD 	8
+#endif
+
+// Use the larger of group sizes to define the shared memory size
+// Threads will do multiple loads over the smaller dimension to fill the memory completely
+groupshared WORK_TYPE SharedMemoryA[GROUP_SIZE_X * GROUP_SIZE_X];
+groupshared WORK_TYPE SharedMemoryB[GROUP_SIZE_X * GROUP_SIZE_X];
+
+[numthreads(GROUP_SIZE_X, GROUP_SIZE_Y, 1)]
+void Gemm(in const uint3 DispatchThreadID : SV_DispatchThreadID, in const uint3 GroupThreadID : SV_GroupThreadID, in const uint3 GroupID : SV_GroupID)
+{
+	// Initialize the result
+	WORK_TYPE Results[LOAD_PER_THREAD];
+	for (int LoadIdx = 0; LoadIdx < LOAD_PER_THREAD; LoadIdx++)
+	{
+        Results[LoadIdx] = 0.0f;
+    }
+	
+	// Calculate the number of steps each group needs to make to cover the full length of the matrix
+	uint NumGroupSteps = (uint)ceil((float)K / (float)GROUP_SIZE_X);
+	
+	// Compute some variables required at multiple places
+	uint SharedMemoryTransposedLineIdx = GroupThreadID.x * GROUP_SIZE_X;
+	uint SharedMemoryStartIdx = GroupThreadID.y * GROUP_SIZE_X + GroupThreadID.x;
+	uint SharedIdxIncrement = GROUP_SIZE_Y * GROUP_SIZE_X;
+	uint GroupStartRow = GroupID.y * GROUP_SIZE_X;
+	uint GroupLoadOffset = GroupThreadID.y * LOAD_PER_THREAD;
+	
+	// Compute the lookup start indices and the increment which needs to be applied for each sub block
+	uint AIdx;
+	uint AIdxOuterIncrement; // Increment added in the group loop
+	uint AIdxInnerIncrement; // Increment added in the load loop
+	if (TransA == 0)
+	{
+		AIdx = (GroupStartRow + GroupThreadID.y) * K + GroupThreadID.x;
+		AIdxOuterIncrement = GROUP_SIZE_X;
+		AIdxInnerIncrement = GROUP_SIZE_Y * K;
+	}
+	else
+	{
+		AIdx = GroupThreadID.x * M + (GroupStartRow + GroupThreadID.y);
+		AIdxOuterIncrement = GROUP_SIZE_X * M;
+		AIdxInnerIncrement = GROUP_SIZE_Y;
+	}
+	
+	uint BIdx;
+	uint BIdxOuterIncrement; // Increment added in the group loop
+	uint BIdxInnerIncrement; // Increment added in the load loop
+	if (TransB == 0)
+	{
+		BIdx = GroupThreadID.y * N + DispatchThreadID.x;
+		BIdxOuterIncrement = GROUP_SIZE_X * N;
+		BIdxInnerIncrement = GROUP_SIZE_Y * N;
+	}
+	else
+	{
+		BIdx = DispatchThreadID.x * K + GroupThreadID.y;
+		BIdxOuterIncrement = GROUP_SIZE_X;
+		BIdxInnerIncrement = GROUP_SIZE_Y;
+	}
+
+	const uint3 StackOffsets = GetMatrixStackOffsets(GroupID);
+	AIdx += StackOffsets.x;
+	BIdx += StackOffsets.y;
+
+	// Multiply and accumulate the blocks
+	uint GroupIdx = 0;
+	for (uint GroupStep = 0; GroupStep < NumGroupSteps; GroupStep++) 
+	{
+		// Prepare some shared variables needed inside this loop
+		uint SharedIdx = SharedMemoryStartIdx;
+		uint AIdxOffset = 0;
+		uint BIdxOffset = 0;
+		
+		uint GroupRow = GroupThreadID.y;
+		uint ThreadIDY = GroupStartRow + GroupRow;
+		
+		// Only sync and update inidces starting with the second iteration
+		if (GroupStep > 0)
+		{
+			// increment indices
+			GroupIdx += GROUP_SIZE_X;
+			AIdx += AIdxOuterIncrement;
+			BIdx += BIdxOuterIncrement;
+			
+			// Make sure computation is done in the whole group before overwriting shared memory
+			GroupMemoryBarrierWithGroupSync();
+		}
+		
+		for (int LoadIdx = 0; LoadIdx < LOAD_PER_THREAD; LoadIdx++)
+		{
+			// Load the shared memory after applying the transpose and set overflow elements to zero
+			WORK_TYPE Temp = 0;
+			if (ThreadIDY < M && GroupIdx + GroupThreadID.x < K)
+			{
+				Temp = READ(A[AIdx + AIdxOffset]);
+			}
+			SharedMemoryA[SharedIdx] = Temp;
+			
+			Temp = 0;
+			if (GroupIdx + GroupRow < K && DispatchThreadID.x < N)
+			{
+				Temp = READ(B[BIdx + BIdxOffset]);
+			}
+			SharedMemoryB[SharedIdx] = Temp;
+			
+			// Increment the thread y coordinate virtually by one group size
+			GroupRow += GROUP_SIZE_Y;
+			ThreadIDY += GROUP_SIZE_Y;
+			
+			// Increment the shared index pointer to advance to the next subblock
+			SharedIdx += SharedIdxIncrement;
+			AIdxOffset += AIdxInnerIncrement;
+			BIdxOffset += BIdxInnerIncrement;
+		}
+		
+		// Make sure all memory has been loaded before running the multiplication
+		GroupMemoryBarrierWithGroupSync();
+		
+		// multiply and accumulate
+		uint StartB = GroupLoadOffset;
+		for (int kGroup = 0; kGroup < GROUP_SIZE_X; kGroup++)
+		{
+			WORK_TYPE AVal = SharedMemoryA[SharedMemoryTransposedLineIdx + kGroup];
+			for (int LoadIdx = 0; LoadIdx < LOAD_PER_THREAD; LoadIdx++)
+			{
+				Results[LoadIdx] += AVal * SharedMemoryB[StartB + LoadIdx];
+			}
+			StartB += GROUP_SIZE_X;
+		}
+	}
+	
+	// Only now return out of bound threads as they were needed to load shared memory above
+	uint Row = GroupStartRow + GroupThreadID.x;
+	if (Row >= M)
+	{
+		return;
+	}
+	
+	// Save back the results
+	uint RowIdx = Row * N;
+	uint ColumnStart = GroupID.x * GROUP_SIZE_X + GroupLoadOffset;
+	for (int LoadIdx = 0; LoadIdx < LOAD_PER_THREAD; LoadIdx++)
+	{
+		uint Column = ColumnStart + LoadIdx;
+		
+		// Return as soon as the column is out of bound as further loops will be outside too
+		if (Column >= N)
+		{
+			return;
+		}
+		
+		Y[StackOffsets.z + RowIdx + Column] = WRITE(Alpha * Results[LoadIdx] + GetBetaTimesC(uint3(Column, Row, 0)));
+    }
+}
+
+#endif
diff --git a/Engine/Plugins/Experimental/NNERuntimeRDG/Shaders/Private/NNEHlslShaders/NNEHlslShadersInstanceNormalization.usf b/Engine/Plugins/Experimental/NNERuntimeRDG/Shaders/Private/NNEHlslShaders/NNEHlslShadersInstanceNormalization.usf
new file mode 100644
--- /dev/null
+++ b/Engine/Plugins/Experimental/NNERuntimeRDG/Shaders/Private/NNEHlslShaders/NNEHlslShadersInstanceNormalization.usf
@@ -0,0 +1,36 @@
+// Copyright Epic Games, Inc. All Rights Reserved.
+ 
+#include "/Engine/Public/Platform.ush"
+ 
+#define WORK_TYPE float
+#define BUFFER_TYPE float
+#define READ(x) x
+#define WRITE(x) x
+ 
+Buffer<WORK_TYPE> Input;
+Buffer<WORK_TYPE> InputScale;
+Buffer<WORK_TYPE> InputBias;
+Buffer<WORK_TYPE> InputMean;
+Buffer<WORK_TYPE> InputInvStdDev;
+RWBuffer<WORK_TYPE> Output;
+uint Num;
+uint ThreadCountX;
+uint InstanceSize;
+uint ChannelSize;
+ 
+[numthreads(THREADGROUP_SIZE_X, 1, 1)]
+void InstanceNormalization(in const uint3 DispatchThreadID : SV_DispatchThreadID)
+{
+	const uint Index = DispatchThreadID.y * ThreadCountX + DispatchThreadID.x;
+
+	if (Index < Num)
+	{
+		const uint InstanceIdx = Index/InstanceSize;
+		const uint ChannelIdx = InstanceIdx%ChannelSize;
+		WORK_TYPE Mean = READ(InputMean[InstanceIdx]);
+		WORK_TYPE InvStdDev = READ(InputInvStdDev[InstanceIdx]);
+		WORK_TYPE Scale = READ(InputScale[ChannelIdx]);
+		WORK_TYPE Bias = READ(InputBias[ChannelIdx]);
+		Output[Index] = WRITE((READ(Input[Index]) - Mean) * Scale * InvStdDev + Bias);
+	}
+}
\ No newline at end of file
diff --git a/Engine/Plugins/Experimental/NNERuntimeRDG/Shaders/Private/NNEHlslShaders/NNEHlslShadersLayerNormalization.usf b/Engine/Plugins/Experimental/NNERuntimeRDG/Shaders/Private/NNEHlslShaders/NNEHlslShadersLayerNormalization.usf
new file mode 100644
--- /dev/null
+++ b/Engine/Plugins/Experimental/NNERuntimeRDG/Shaders/Private/NNEHlslShaders/NNEHlslShadersLayerNormalization.usf
@@ -0,0 +1,78 @@
+// Copyright Epic Games, Inc. All Rights Reserved.
+ 
+#include "/Engine/Public/Platform.ush"
+#include "NNEHlslShadersBroadcastHelper.ush"
+
+#define WORK_TYPE float
+#define BUFFER_TYPE float
+#define READ(x) x
+#define WRITE(x) x
+
+Buffer<BUFFER_TYPE> Input;
+uint4 InputTensorInfo[NUM_DIMENSIONS];
+Buffer<BUFFER_TYPE> InputScale;
+uint4 ScaleTensorInfo[NUM_DIMENSIONS];
+Buffer<BUFFER_TYPE> InputBias;
+uint4 BiasTensorInfo[NUM_DIMENSIONS];
+Buffer<WORK_TYPE> InputMean;
+Buffer<WORK_TYPE> InputInvStdDev;
+RWBuffer<BUFFER_TYPE> Output;
+
+float Epsilon;
+int Axis;
+uint ThreadCountX;
+uint LayerSize;
+uint Num;
+
+#define STRIDE_IDX 0
+
+#define STATIC_LOOP(Var, From, To) \
+[unroll] \
+for(uint Var = From; Var < To; ++Var)
+
+void GetDimIterator(const uint GlobalIdx, const uint4 TensorInfo[NUM_DIMENSIONS], out uint DimIterator[NUM_DIMENSIONS])
+{
+	uint Offset = GlobalIdx;
+	STATIC_LOOP(DimIdx, 0, NUM_DIMENSIONS)
+	{
+		uint Remainder;
+		DivMod(Offset, TensorInfo[DimIdx][STRIDE_IDX], DimIterator[DimIdx], Remainder);
+		Offset = Remainder;
+	}
+}
+
+uint GetInnerIdx(const uint DimIterator[NUM_DIMENSIONS], const uint4 TensorInfo[NUM_DIMENSIONS], const uint Axis)
+{
+	uint Offset = 0;
+	for(uint DimIdx = 0; DimIdx < NUM_DIMENSIONS - Axis; ++DimIdx)
+	{
+		Offset += DimIterator[DimIdx + Axis] * TensorInfo[DimIdx][STRIDE_IDX];
+	}
+	return Offset;
+}
+
+
+[numthreads(THREADGROUP_SIZE, 1, 1)]
+void LayerNormalization(in const uint3 DispatchThreadID : SV_DispatchThreadID)
+{
+	const uint ElementIdx = DispatchThreadID.y * ThreadCountX + DispatchThreadID.x;
+	if(ElementIdx < Num)
+	{
+		const uint LayerIdx = ElementIdx / LayerSize;
+		
+		WORK_TYPE Mean = InputMean[LayerIdx];
+		WORK_TYPE InvStdDev = InputInvStdDev[LayerIdx];
+
+		BUFFER_TYPE Normalized = (BUFFER_TYPE) ( (READ(Input[ElementIdx]) - Mean) * InvStdDev );
+
+		uint InputDimIterator[NUM_DIMENSIONS];
+		GetDimIterator(ElementIdx, InputTensorInfo, InputDimIterator);
+
+		BUFFER_TYPE NormalizedScaled = Normalized * READ( InputScale[GetInnerIdx(InputDimIterator, ScaleTensorInfo, Axis)]);
+		#if HAS_B == 0
+			Output[ElementIdx] = WRITE( NormalizedScaled );
+		#else
+			Output[ElementIdx] = WRITE( NormalizedScaled + READ( InputBias[GetInnerIdx(InputDimIterator, BiasTensorInfo, Axis)] ) );
+		#endif
+	}
+}
diff --git a/Engine/Plugins/Experimental/NNERuntimeRDG/Shaders/Private/NNEHlslShaders/NNEHlslShadersNeuralPostProcessing.usf b/Engine/Plugins/Experimental/NNERuntimeRDG/Shaders/Private/NNEHlslShaders/NNEHlslShadersNeuralPostProcessing.usf
new file mode 100644
--- /dev/null
+++ b/Engine/Plugins/Experimental/NNERuntimeRDG/Shaders/Private/NNEHlslShaders/NNEHlslShadersNeuralPostProcessing.usf
@@ -0,0 +1,154 @@
+// Copyright Epic Games, Inc. All Rights Reserved.
+
+#include "/Engine/Public/Platform.ush"
+#include "/Engine/Private/Common.ush"
+
+int InputTextureWidth;
+int InputTextureHeight;
+Texture2D InputTexture;
+SamplerState InputTextureSampler;
+RWBuffer<float> AccumulationBuffer;
+
+int InputBufferWidth;
+int InputBufferHeight;
+RWBuffer<float> InputBuffer;
+
+int OutputBufferWidth;
+int OutputBufferHeight;
+RWBuffer<float> OutputBuffer;
+
+float Weight;
+float RangeScale;
+
+[numthreads(THREAD_GROUP_SIZE, THREAD_GROUP_SIZE, 1)]
+void ReadInput(in const uint3 DispatchThreadID : SV_DispatchThreadID)
+{
+	if (DispatchThreadID.x >= InputTextureWidth || DispatchThreadID.y >= InputTextureHeight)
+	{
+		return;
+	}
+
+	float4 InputColor = InputTexture.SampleLevel(InputTextureSampler, float2((float(DispatchThreadID.x) + 0.5) / float(InputTextureWidth), (float(DispatchThreadID.y) + 0.5) / float(InputTextureHeight)), 0);
+
+	int Idx = (DispatchThreadID.y * InputTextureWidth + DispatchThreadID.x) * 3;
+	AccumulationBuffer[Idx + 0] = Weight * InputColor.r;
+	AccumulationBuffer[Idx + 1] = Weight * InputColor.g;
+	AccumulationBuffer[Idx + 2] = Weight * InputColor.b;
+}
+
+[numthreads(THREAD_GROUP_SIZE, THREAD_GROUP_SIZE, 1)]
+void PreStep(in const uint3 DispatchThreadID : SV_DispatchThreadID)
+{
+	if (DispatchThreadID.x >= InputBufferWidth || DispatchThreadID.y >= InputBufferHeight)
+	{
+		return;
+	}
+
+	float WidthFactor = (float)InputTextureWidth / (float)InputBufferWidth;
+	int LoX = (int)floor(WidthFactor * (float)DispatchThreadID.x);
+	int HiX = (int)ceil(WidthFactor * (float)(DispatchThreadID.x + 1));
+
+	float HeightFactor = (float)InputTextureHeight / (float)InputBufferHeight;
+	int LoY = (int)floor(HeightFactor * (float)DispatchThreadID.y);
+	int HiY = (int)ceil(HeightFactor * (float)(DispatchThreadID.y + 1));
+
+	float4 Result = float4(0.0, 0.0, 0.0, 0.0);
+	float Div = 0.0;
+	for (int x = LoX; x < HiX; x++)
+	{
+		if (x < InputTextureWidth)
+		{
+			for (int y = LoY; y < HiY; y++)
+			{
+				if (y < InputTextureHeight)
+				{
+					Result += InputTexture.SampleLevel(InputTextureSampler, float2((x + 0.5) / float(InputTextureWidth), (y + 0.5) / float(InputTextureHeight)), 0);
+					Div += 1.0;
+				}
+			}
+		}
+	}
+	Result /= Div;
+
+	int Idx = DispatchThreadID.y * InputBufferWidth + DispatchThreadID.x;
+	int Offset1 = InputBufferWidth * InputBufferHeight;
+	int Offset2 = Offset1 + Offset1;
+	InputBuffer[Idx] = RangeScale * Result.r;
+	InputBuffer[Idx + Offset1] = RangeScale * Result.g;
+	InputBuffer[Idx + Offset2] = RangeScale * Result.b;
+}
+
+[numthreads(THREAD_GROUP_SIZE, THREAD_GROUP_SIZE, 1)]
+void PostStep(in const uint3 DispatchThreadID : SV_DispatchThreadID)
+{
+	if (DispatchThreadID.x >= InputTextureWidth || DispatchThreadID.y >= InputTextureHeight)
+	{
+		return;
+	}
+
+	float WidthFactor = (float)OutputBufferWidth / (float)InputTextureWidth;
+	float X = WidthFactor * ((float)DispatchThreadID.x + 0.5);
+
+	float HeightFactor = (float)OutputBufferHeight / (float)InputTextureHeight;
+	float Y = HeightFactor * ((float)DispatchThreadID.y + 0.5);
+
+	int Offset1 = OutputBufferWidth * OutputBufferHeight;
+	int Offset2 = Offset1 + Offset1;
+
+	float3 Result;
+#if INTERPOLATE == 0
+	int Idx = (int)Y * OutputBufferWidth + (int)X;
+	Result.r = OutputBuffer[Idx];
+	Result.g = OutputBuffer[Idx + Offset1];
+	Result.b = OutputBuffer[Idx + Offset2];
+#else
+	int LoX = (int)clamp(floor(X - 0.5), 0.0, float(OutputBufferWidth - 1));
+	int HiX = (int)clamp(floor(X + 0.5), 0.0, float(OutputBufferWidth - 1));
+
+	int LoY = (int)clamp(floor(Y - 0.5), 0.0, float(OutputBufferHeight - 1));
+	int HiY = (int)clamp(floor(Y + 0.5), 0.0, float(OutputBufferHeight - 1));
+
+	int Idx;
+	
+	Idx = LoY * OutputBufferWidth + LoX;
+	float3 A = float3(OutputBuffer[Idx], OutputBuffer[Idx + Offset1], OutputBuffer[Idx + Offset2]);
+
+	Idx = LoY * OutputBufferWidth + HiX;
+	float3 B = float3(OutputBuffer[Idx], OutputBuffer[Idx + Offset1], OutputBuffer[Idx + Offset2]);
+
+	Idx = HiY * OutputBufferWidth + LoX;
+	float3 C = float3(OutputBuffer[Idx], OutputBuffer[Idx + Offset1], OutputBuffer[Idx + Offset2]);
+
+	Idx = HiY * OutputBufferWidth + HiX;
+	float3 D = float3(OutputBuffer[Idx], OutputBuffer[Idx + Offset1], OutputBuffer[Idx + Offset2]);
+
+	float Alpha = clamp((X - 0.5) - ((float)LoX), 0.0, 1.0);
+	float Beta = clamp((Y - 0.5) - ((float)LoY), 0.0, 1.0);
+	
+	Result = (1.0 - Beta) * ((1.0 - Alpha) * A + Alpha * B) + Beta * ((1.0 - Alpha) * C + Alpha * D);
+#endif
+
+	int ResultIdx = (DispatchThreadID.y * InputTextureWidth + DispatchThreadID.x) * 3;
+#if OVERWRITE == 0
+	AccumulationBuffer[ResultIdx + 0] += Weight * Result.r / RangeScale;
+	AccumulationBuffer[ResultIdx + 1] += Weight * Result.g / RangeScale;
+	AccumulationBuffer[ResultIdx + 2] += Weight * Result.b / RangeScale;
+#else
+	AccumulationBuffer[ResultIdx + 0] = Weight * Result.r / RangeScale;
+	AccumulationBuffer[ResultIdx + 1] = Weight * Result.g / RangeScale;
+	AccumulationBuffer[ResultIdx + 2] = Weight * Result.b / RangeScale;
+#endif
+}
+
+float4 WriteOutput(float4 Position : SV_POSITION) : SV_Target0
+{
+	int Idx = ((int)Position.y * InputTextureWidth + (int)Position.x) * 3;
+
+	float4 Result;
+	Result.r = AccumulationBuffer[Idx + 0];
+	Result.g = AccumulationBuffer[Idx + 1];
+	Result.b = AccumulationBuffer[Idx + 2];
+	Result.a = 1.0;
+
+	return Result;
+}
\ No newline at end of file
diff --git a/Engine/Plugins/Experimental/NNERuntimeRDG/Shaders/Private/NNEHlslShaders/NNEHlslShadersPad.usf b/Engine/Plugins/Experimental/NNERuntimeRDG/Shaders/Private/NNEHlslShaders/NNEHlslShadersPad.usf
new file mode 100644
--- /dev/null
+++ b/Engine/Plugins/Experimental/NNERuntimeRDG/Shaders/Private/NNEHlslShaders/NNEHlslShadersPad.usf
@@ -0,0 +1,93 @@
+// Copyright Epic Games, Inc. All Rights Reserved.
+
+#include "/Engine/Public/Platform.ush"
+#include "NNEHlslShadersBroadcastHelper.ush"
+
+Buffer<float> Input;
+RWBuffer<float> Output;
+uint4 TensorInfo[NUM_DIMENSIONS];
+float Value;
+uint Num;
+uint ThreadCountX;
+
+#define INPUT_STRIDE 0
+#define OUTPUT_STRIDE 1
+#define INPUT_SIZE 2
+#define PRE_PAD 3
+
+// Must correspond to EPadMode defined in NNEHlslShadersPadCS.h
+#define CONSTANT_MODE 0
+#define REFLECT_MODE 1
+#define EDGE_MODE 2
+
+#define STATIC_LOOP(Var, From, To) \
+[unroll] \
+for(uint Var = From; Var < To; ++Var)
+
+void GetOutputDimIterator(const uint GlobalIdx, out uint DimIterator[NUM_DIMENSIONS])
+{
+	uint Offset = GlobalIdx;
+	STATIC_LOOP(DimIdx, 0, NUM_DIMENSIONS)
+	{
+		uint Remainder;
+		DivMod(Offset, TensorInfo[DimIdx][OUTPUT_STRIDE], DimIterator[DimIdx], Remainder);
+		Offset = Remainder;
+	}
+}
+
+uint GetInputGlobalIdx(const uint DimIterator[NUM_DIMENSIONS])
+{
+	uint Offset = 0;
+	STATIC_LOOP(DimIdx, 0, NUM_DIMENSIONS)
+	{
+		Offset += DimIterator[DimIdx] * TensorInfo[DimIdx][INPUT_STRIDE];
+	}
+	return Offset;
+}
+
+[numthreads(THREADGROUP_SIZE_X, 1, 1)]
+void Pad(in const uint3 DispatchThreadID : SV_DispatchThreadID)
+{
+	const uint Index = DispatchThreadID.y * ThreadCountX + DispatchThreadID.x;
+	
+	if (Index < Num)
+	{
+		uint OutputDimIterator[NUM_DIMENSIONS];
+		GetOutputDimIterator(Index, OutputDimIterator);
+
+		uint InputDimIterator[NUM_DIMENSIONS];
+
+		for(uint DimIdx = 0; DimIdx < NUM_DIMENSIONS; ++DimIdx)
+		{
+			const int InputDimSize = (int) TensorInfo[DimIdx][INPUT_SIZE];
+			
+			int PrePad = asint(TensorInfo[DimIdx][PRE_PAD]); // Decode PrePad to int (can be negative)
+			int CurDimIterator = (int) OutputDimIterator[DimIdx] - PrePad;
+			
+			if (CurDimIterator < 0 || CurDimIterator >= InputDimSize)
+			{
+				#if MODE == CONSTANT_MODE
+					Output[Index] = Value;
+					return;
+				#elif MODE == REFLECT_MODE
+					if (CurDimIterator >= InputDimSize)
+					{
+						const int OverCnt = CurDimIterator - (InputDimSize - 2);
+						CurDimIterator = InputDimSize - OverCnt;
+					}
+					else if (CurDimIterator < 0)
+					{
+						CurDimIterator = -CurDimIterator;
+					}
+				#endif
+				//MODE == EDGE_MODE + clamping large reflect indices
+				CurDimIterator = min(CurDimIterator, InputDimSize - 1);
+				CurDimIterator = max(CurDimIterator, 0);
+			}
+
+			InputDimIterator[DimIdx] = (uint) CurDimIterator;
+		}
+
+		Output[Index] = Input[GetInputGlobalIdx(InputDimIterator)];
+	}
+}
diff --git a/Engine/Plugins/Experimental/NNERuntimeRDG/Shaders/Private/NNEHlslShaders/NNEHlslShadersPool.usf b/Engine/Plugins/Experimental/NNERuntimeRDG/Shaders/Private/NNEHlslShaders/NNEHlslShadersPool.usf
new file mode 100644
--- /dev/null
+++ b/Engine/Plugins/Experimental/NNERuntimeRDG/Shaders/Private/NNEHlslShaders/NNEHlslShadersPool.usf
@@ -0,0 +1,128 @@
+// Copyright Epic Games, Inc. All Rights Reserved.
+
+#include "/Engine/Public/Platform.ush"
+#include "NNEHlslShadersBroadcastHelper.ush"
+
+#define NUM_DIMENSION (NUM_SPATIAL_DIMENSIONS+2)
+
+Buffer<float> Input;
+RWBuffer<float> Output;
+uint4 TensorInfo[NUM_DIMENSION];
+uint4 SpatialInfo[NUM_SPATIAL_DIMENSIONS];
+uint Num;
+uint ThreadCountX;
+uint KernelVolume;
+
+#define OUTPUT_STRIDE 0
+#define INPUT_STRIDE 1
+#define INPUT_SIZE 2
+
+#define SPATIAL_STRIDE 0
+#define SPATIAL_KERNEL 1
+#define SPATIAL_PAD_START 2
+#define SPATIAL_DILATION 3
+
+#define MAX_POOL 0
+#define AVERAGE_POOL 1
+
+[numthreads(THREADGROUP_SIZE_X, 1, 1)]
+void Pool(in const uint3 DispatchThreadID : SV_DispatchThreadID)
+{
+	const uint Index = DispatchThreadID.y * ThreadCountX + DispatchThreadID.x;
+	const uint SpatialDimOffset = 2;
+	
+	if (Index < Num)
+	{
+		//Compute indices on output tensors
+		uint Offset = Index;
+		uint OutputIndices[NUM_DIMENSION];
+		for (uint dim = 0; dim < NUM_DIMENSION; ++dim)
+		{
+			uint OutDimIdx, R;
+			DivMod(Offset, TensorInfo[dim][OUTPUT_STRIDE], OutputIndices[dim], R);
+			Offset = R;
+		}
+
+		//Compute spatial region to apply pooling on input
+		uint SpatialDimStarts[NUM_SPATIAL_DIMENSIONS];
+		uint SpatialDimEnds[NUM_SPATIAL_DIMENSIONS];
+		for (uint spatialDim = 0; spatialDim < NUM_SPATIAL_DIMENSIONS; ++spatialDim)
+		{
+			const uint Dim = SpatialDimOffset + spatialDim;
+			const int SpatialDimStart = ((int)OutputIndices[Dim] * (int)SpatialInfo[spatialDim][SPATIAL_STRIDE]) - (int)SpatialInfo[spatialDim][SPATIAL_PAD_START];
+			const int SpatialDimEnd = SpatialDimStart + (int)SpatialInfo[spatialDim][SPATIAL_KERNEL] * (int)SpatialInfo[spatialDim][SPATIAL_DILATION];
+
+			SpatialDimStarts[spatialDim] = (uint)max(SpatialDimStart, (int)0);
+			SpatialDimEnds[spatialDim] = (uint)min(SpatialDimEnd, (int)TensorInfo[Dim][INPUT_SIZE]);
+		}
+
+		const uint BaseOffset = (TensorInfo[0][INPUT_STRIDE] * OutputIndices[0]) + (TensorInfo[1][INPUT_STRIDE] * OutputIndices[1]);
+		uint SpatialDimIterator[NUM_SPATIAL_DIMENSIONS];
+
+		//Apply pooling, looping other the input for all spatial dimensions of the input kernel
+		float PooledValue = 0;
+		int Count = 0;
+		
+		for (SpatialDimIterator[0] = SpatialDimStarts[0]; SpatialDimIterator[0] < SpatialDimEnds[0]; SpatialDimIterator[0] += SpatialInfo[0][SPATIAL_DILATION])
+		{
+		#if (NUM_SPATIAL_DIMENSIONS > 1)
+		for (SpatialDimIterator[1] = SpatialDimStarts[1]; SpatialDimIterator[1] < SpatialDimEnds[1]; SpatialDimIterator[1] += SpatialInfo[1][SPATIAL_DILATION])
+		{
+		#endif
+		#if (NUM_SPATIAL_DIMENSIONS > 2)
+		for (SpatialDimIterator[2] = SpatialDimStarts[2]; SpatialDimIterator[2] < SpatialDimEnds[2]; SpatialDimIterator[2] += SpatialInfo[2][SPATIAL_DILATION])
+		{
+		#endif
+		#if (NUM_SPATIAL_DIMENSIONS > 3)
+		for (SpatialDimIterator[3] = SpatialDimStarts[3]; SpatialDimIterator[3] < SpatialDimEnds[3]; SpatialDimIterator[3] += SpatialInfo[3][SPATIAL_DILATION])
+		{
+		#endif
+		#if (NUM_SPATIAL_DIMENSIONS > 4)
+		for (SpatialDimIterator[4] = SpatialDimStarts[4]; SpatialDimIterator[4] < SpatialDimEnds[4]; SpatialDimIterator[4] += SpatialInfo[4][SPATIAL_DILATION])
+		{
+		#endif
+		#if (NUM_SPATIAL_DIMENSIONS > 5)
+		for (SpatialDimIterator[5] = SpatialDimStarts[5]; SpatialDimIterator[5] < SpatialDimEnds[5]; SpatialDimIterator[5] += SpatialInfo[5][SPATIAL_DILATION])
+		{
+		#endif
+
+			uint InputIndex = BaseOffset;
+			[unroll]
+			for (uint spatialDim = 0; spatialDim < NUM_SPATIAL_DIMENSIONS; ++spatialDim)
+			{
+				InputIndex += TensorInfo[spatialDim+2][INPUT_STRIDE] * SpatialDimIterator[spatialDim];
+			}
+			float CurrentValue = Input[InputIndex];
+			#if (POOL_OPERATOR_TYPE == MAX_POOL)
+				//float MaxValue = max(PooledValue, CurrentValue); //Is not working properly on DX12 SM6 thus the ternary op for MaxValue
+				float MaxValue = (PooledValue > CurrentValue) ? PooledValue : CurrentValue;
+				PooledValue = (Count == 0) ? CurrentValue : MaxValue;
+			#else //POOL_OPERATOR_TYPE == AVERAGE_POOL
+				PooledValue += CurrentValue;
+			#endif 
+			++Count;
+
+		#if (NUM_SPATIAL_DIMENSIONS > 5)
+		}
+		#endif
+		#if (NUM_SPATIAL_DIMENSIONS > 4)
+		}
+		#endif
+		#if (NUM_SPATIAL_DIMENSIONS > 3)
+		}
+		#endif
+		#if (NUM_SPATIAL_DIMENSIONS > 2)
+		}
+		#endif
+		#if (NUM_SPATIAL_DIMENSIONS > 1)
+		}
+		#endif
+		}
+
+		#if (POOL_OPERATOR_TYPE == AVERAGE_POOL)
+			PooledValue /= (KernelVolume == 0) ? (float)Count : (float)KernelVolume;
+		#endif
+
+		Output[Index] = PooledValue;
+	}
+}
diff --git a/Engine/Plugins/Experimental/NNERuntimeRDG/Shaders/Private/NNEHlslShaders/NNEHlslShadersReduce.usf b/Engine/Plugins/Experimental/NNERuntimeRDG/Shaders/Private/NNEHlslShaders/NNEHlslShadersReduce.usf
new file mode 100644
--- /dev/null
+++ b/Engine/Plugins/Experimental/NNERuntimeRDG/Shaders/Private/NNEHlslShaders/NNEHlslShadersReduce.usf
@@ -0,0 +1,222 @@
+// Copyright Epic Games, Inc. All Rights Reserved.
+
+#include "/Engine/Public/Platform.ush"
+
+#define WORK_TYPE float
+#define BUFFER_TYPE float
+#define READ(x) x
+#define WRITE(x) x
+
+Buffer<WORK_TYPE> Input;
+RWBuffer<WORK_TYPE> Output;
+RWBuffer<BUFFER_TYPE> Output2;
+int NumElemBeforeAxis;
+int AxisSize;
+int NumElemAfterAxis;
+uint4 DispatchIdxAndStride;
+float Epsilon;
+
+// Must correspond to EReduceOperatorType defined in NNEHlslShadersReduceCS.h
+#define TYPE_AVERAGE 0
+#define TYPE_L1 1
+#define TYPE_L2 2
+#define TYPE_LOGSUMEXP 3
+#define TYPE_MAX 4
+#define TYPE_MIN 5
+#define TYPE_PROD 6
+#define TYPE_SUM 7
+#define TYPE_SUMEXP 8
+#define TYPE_AVERAGEINVSTDDEV 9
+
+#if REDUCE_OPERATOR_TYPE == TYPE_AVERAGEINVSTDDEV
+	groupshared WORK_TYPE SharedMemoryMean[THREADGROUP_SIZE];
+	groupshared WORK_TYPE SharedMemoryMean2[THREADGROUP_SIZE];
+	groupshared WORK_TYPE SharedMemoryCount[THREADGROUP_SIZE];
+#else //REDUCE_OPERATOR_TYPE != TYPE_AVERAGEINVSTD
+	groupshared WORK_TYPE SharedMemory[THREADGROUP_SIZE];
+
+	WORK_TYPE InitValue(WORK_TYPE FirstValue)
+	{
+		#if REDUCE_OPERATOR_TYPE == TYPE_AVERAGE
+			return 0;
+		#elif REDUCE_OPERATOR_TYPE == TYPE_L1
+			return 0;
+		#elif REDUCE_OPERATOR_TYPE == TYPE_L2
+			return 0;
+		#elif REDUCE_OPERATOR_TYPE == TYPE_LOGSUMEXP
+			return 0;
+		#elif REDUCE_OPERATOR_TYPE == TYPE_MAX
+			return FirstValue;
+		#elif REDUCE_OPERATOR_TYPE == TYPE_MIN
+			return FirstValue;
+		#elif REDUCE_OPERATOR_TYPE == TYPE_PROD
+			return 1;
+		#elif REDUCE_OPERATOR_TYPE == TYPE_SUM
+			return 0;
+		#elif REDUCE_OPERATOR_TYPE == TYPE_SUMEXP
+			return 0;
+		#endif
+	}
+
+	WORK_TYPE ElementOp(WORK_TYPE Value)
+	{
+		#if REDUCE_OPERATOR_TYPE == TYPE_AVERAGE
+			return Value;
+		#elif REDUCE_OPERATOR_TYPE == TYPE_L1
+			return abs(Value);
+		#elif REDUCE_OPERATOR_TYPE == TYPE_L2
+			return Value*Value;
+		#elif REDUCE_OPERATOR_TYPE == TYPE_LOGSUMEXP
+			return exp(Value);
+		#elif REDUCE_OPERATOR_TYPE == TYPE_MAX
+			return Value;
+		#elif REDUCE_OPERATOR_TYPE == TYPE_MIN
+			return Value;
+		#elif REDUCE_OPERATOR_TYPE == TYPE_PROD
+			return Value;
+		#elif REDUCE_OPERATOR_TYPE == TYPE_SUM
+			return Value;
+		#elif REDUCE_OPERATOR_TYPE == TYPE_SUMEXP
+			return exp(Value);
+		#endif
+	}
+
+	WORK_TYPE CombineOp(WORK_TYPE ValueA, WORK_TYPE ValueB)
+	{
+		#if REDUCE_OPERATOR_TYPE == TYPE_AVERAGE
+			return ValueA+ValueB;
+		#elif REDUCE_OPERATOR_TYPE == TYPE_L1
+			return ValueA+ValueB;
+		#elif REDUCE_OPERATOR_TYPE == TYPE_L2
+			return ValueA+ValueB;
+		#elif REDUCE_OPERATOR_TYPE == TYPE_LOGSUMEXP
+			return ValueA+ValueB;
+		#elif REDUCE_OPERATOR_TYPE == TYPE_MAX
+			return max(ValueA, ValueB);
+		#elif REDUCE_OPERATOR_TYPE == TYPE_MIN
+			return min(ValueA, ValueB);
+		#elif REDUCE_OPERATOR_TYPE == TYPE_PROD
+			return ValueA*ValueB;
+		#elif REDUCE_OPERATOR_TYPE == TYPE_SUM
+			return ValueA+ValueB;
+		#elif REDUCE_OPERATOR_TYPE == TYPE_SUMEXP
+			return ValueA+ValueB;
+		#endif
+	}
+
+	WORK_TYPE FinalOp(WORK_TYPE Value)
+	{
+		#if REDUCE_OPERATOR_TYPE == TYPE_AVERAGE
+			return Value/AxisSize;
+		#elif REDUCE_OPERATOR_TYPE == TYPE_L1
+			return Value;
+		#elif REDUCE_OPERATOR_TYPE == TYPE_L2
+			return sqrt(Value);
+		#elif REDUCE_OPERATOR_TYPE == TYPE_LOGSUMEXP
+			return log(Value);
+		#elif REDUCE_OPERATOR_TYPE == TYPE_MAX
+			return Value;
+		#elif REDUCE_OPERATOR_TYPE == TYPE_MIN
+			return Value;
+		#elif REDUCE_OPERATOR_TYPE == TYPE_PROD
+			return Value;
+		#elif REDUCE_OPERATOR_TYPE == TYPE_SUM
+			return Value;
+		#elif REDUCE_OPERATOR_TYPE == TYPE_SUMEXP
+			return Value;
+		#endif
+	}
+#endif //REDUCE_OPERATOR_TYPE != TYPE_AVERAGEINVSTD
+
+[numthreads(THREADGROUP_SIZE, 1, 1)]
+void Reduce(in const uint3 DispatchThreadID : SV_DispatchThreadID, in const uint3 GroupThreadID : SV_GroupThreadID)
+{
+	const uint IdxAfterAxis = DispatchIdxAndStride.x * DispatchIdxAndStride.z +  DispatchThreadID.y;
+	const uint IdxBeforeAxis =  DispatchIdxAndStride.y * DispatchIdxAndStride.w +  DispatchThreadID.z;
+	const int InputStartIdx = (IdxBeforeAxis * AxisSize * NumElemAfterAxis) + IdxAfterAxis;
+	const int OutputIdx     = (IdxBeforeAxis * 1        * NumElemAfterAxis) + IdxAfterAxis;
+
+	// Reduction using sharedmem and THREADGROUP_SIZE threads along the reduced dimension
+	{
+		#if REDUCE_OPERATOR_TYPE == TYPE_AVERAGEINVSTDDEV
+			WORK_TYPE Mean = 0;
+			WORK_TYPE Mean2 = 0;
+			int Count = 0;
+		#else //REDUCE_OPERATOR_TYPE != TYPE_AVERAGEINVSTDDEV
+			WORK_TYPE FirstValue = READ(Input[InputStartIdx]);
+			WORK_TYPE Reduced = InitValue(FirstValue);
+		#endif
+
+
+			for (int AxisIdx = GroupThreadID.x; AxisIdx < AxisSize; AxisIdx += THREADGROUP_SIZE)
+			{
+				WORK_TYPE Value = READ(Input[(AxisIdx * NumElemAfterAxis) + InputStartIdx]);
+			
+				#if REDUCE_OPERATOR_TYPE == TYPE_AVERAGEINVSTDDEV
+					// Calculate mean and variance using Welford's online algorithm
+					// https://en.wikipedia.org/wiki/Algorithms_for_calculating_variance
+					Count += 1;
+					WORK_TYPE Delta = Value - Mean;
+					Mean += Delta / (WORK_TYPE)Count;
+					WORK_TYPE Delta2 = Value - Mean;
+					Mean2 += Delta * Delta2;
+				#else //REDUCE_OPERATOR_TYPE != TYPE_AVERAGEINVSTDDEV
+					Value = ElementOp(Value);
+					Reduced = CombineOp(Reduced, Value);
+				#endif
+			}
+
+		#if REDUCE_OPERATOR_TYPE == TYPE_AVERAGEINVSTDDEV
+			SharedMemoryMean[GroupThreadID.x] = Mean;
+			SharedMemoryMean2[GroupThreadID.x] = Mean2;
+			SharedMemoryCount[GroupThreadID.x] = Count;
+		#else //REDUCE_OPERATOR_TYPE != TYPE_AVERAGEINVSTDDEV
+			SharedMemory[GroupThreadID.x] = Reduced;
+		#endif
+	}
+
+	// Wait until all threads wrote result to shared memory
+	GroupMemoryBarrierWithGroupSync();
+
+	// Use one thread to calculate final result
+	if (GroupThreadID.x == 0)
+	{
+		#if REDUCE_OPERATOR_TYPE == TYPE_AVERAGEINVSTDDEV
+			WORK_TYPE Mean = 0;
+			WORK_TYPE Mean2 = 0;
+			int Count = 0;
+		#else //REDUCE_OPERATOR_TYPE != TYPE_AVERAGEINVSTDDEV
+			WORK_TYPE FirstCurrReduced = SharedMemory[0];
+			WORK_TYPE Reduced = InitValue(FirstCurrReduced);
+		#endif
+
+			UNROLL
+			for (int Idx = 0; Idx < THREADGROUP_SIZE; Idx++)
+			{
+				#if REDUCE_OPERATOR_TYPE == TYPE_AVERAGEINVSTDDEV
+					WORK_TYPE CurrMean = SharedMemoryMean[Idx];
+					WORK_TYPE CurrMean2 = SharedMemoryMean2[Idx];
+					int CurrCount = SharedMemoryCount[Idx];
+					int NewCount = Count + CurrCount;
+					WORK_TYPE Delta = CurrMean - Mean;
+					Mean = (CurrCount * CurrMean + Count * Mean) / NewCount;
+					Mean2 += CurrMean2 + Delta * Delta * (WORK_TYPE)Count * (WORK_TYPE)CurrCount / (WORK_TYPE)NewCount;
+					Count = NewCount;
+				#else //REDUCE_OPERATOR_TYPE != TYPE_AVERAGEINVSTDDEV
+					WORK_TYPE CurrReduced = SharedMemory[Idx];
+					Reduced = CombineOp(Reduced, CurrReduced);
+				#endif
+			}
+
+		// Write reduced result
+		#if REDUCE_OPERATOR_TYPE == TYPE_AVERAGEINVSTDDEV
+			WORK_TYPE Variance = Mean2 / (WORK_TYPE)Count;
+			WORK_TYPE InvStdev = (WORK_TYPE)1 / sqrt(Variance + Epsilon);
+			Output[OutputIdx] = Mean;
+			Output2[OutputIdx] = InvStdev;
+		#else //REDUCE_OPERATOR_TYPE != TYPE_AVERAGEINVSTDDEV
+			Reduced = FinalOp(Reduced);
+			Output[OutputIdx] = Reduced;
+		#endif
+	}
+}
\ No newline at end of file
diff --git a/Engine/Plugins/Experimental/NNERuntimeRDG/Shaders/Private/NNEHlslShaders/NNEHlslShadersResize.usf b/Engine/Plugins/Experimental/NNERuntimeRDG/Shaders/Private/NNEHlslShaders/NNEHlslShadersResize.usf
new file mode 100644
--- /dev/null
+++ b/Engine/Plugins/Experimental/NNERuntimeRDG/Shaders/Private/NNEHlslShaders/NNEHlslShadersResize.usf
@@ -0,0 +1,213 @@
+// Copyright Epic Games, Inc. All Rights Reserved.
+
+#include "/Engine/Public/Platform.ush"
+#include "NNEHlslShadersBroadcastHelper.ush"
+#include "/Engine/Private/DoubleFloat.ush"
+
+static const uint NumPointsLinear = pow(2,NUM_DIMENSIONS);
+
+Buffer<float> Input;
+uint4 InputTensorInfo[NUM_DIMENSIONS];
+RWBuffer<float> Output;
+uint4 OutputTensorInfo[NUM_DIMENSIONS];
+float4 ScalesData[NUM_DIMENSIONS];
+uint Num;
+uint ThreadCountX;
+
+#define STRIDE_IDX 0
+#define SIZE_IDX 1
+#define ADJUSTMENT_IDX 2
+#define SCALES_IDX 0
+#define ROI_START_IDX 2
+#define ROI_END_IDX 3
+
+// Corresponds to ECoordTransMode
+#define COORD_TRANS_HALF_PIXEL 0
+#define COORD_TRANS_HALF_PIXEL_SYMMETRIC 1
+#define COORD_TRANS_PYTORCH_HALF_PIXEL 2
+#define COORD_TRANS_ALIGN_CORNERS 3
+#define COORD_TRANS_ASYMMETRIC 4
+#define COORD_TRANS_TF_HALF_PIXEL_FOR_NN 5
+#define COORD_TRANS_TF_CROP_AND_RESIZE 6
+
+// Corresponds to ECoordTransMode enum
+#define NEAREST_MODE_ROUND_PREFER_FLOOR 0
+#define NEAREST_MODE_ROUND_PREFER_CEIL 1
+#define NEAREST_MODE_FLOOR 2
+#define NEAREST_MODE_CEIL 3
+
+// Corresponds to EMode enum
+#define MODE_NEAREST 0
+#define MODE_LINEAR 1
+#define MODE_CUBIC 2
+
+#define STATIC_LOOP(Var, From, To) \
+[unroll] \
+for(uint Var = From; Var < To; ++Var)
+
+void FromIdxToIterator(const uint Idx, const uint4 TensorInfo[NUM_DIMENSIONS], out uint Iterator[NUM_DIMENSIONS])
+{
+	uint Offset = Idx;
+	STATIC_LOOP(DimIdx, 0, NUM_DIMENSIONS)
+	{
+		uint Remainder;
+		DivMod(Offset, TensorInfo[DimIdx][STRIDE_IDX], Iterator[DimIdx], Remainder);
+		Offset = Remainder;
+	}
+}
+
+uint FromIteratorToIdx(const uint4 TensorInfo[NUM_DIMENSIONS], const uint Iterator[NUM_DIMENSIONS])
+{
+	uint Offset = 0;
+	STATIC_LOOP(DimIdx, 0, NUM_DIMENSIONS)
+	{
+		Offset += Iterator[DimIdx] * TensorInfo[DimIdx][STRIDE_IDX];
+	}
+	return Offset;
+}
+
+float TransformCoord(float Resized, uint DimIdx)
+{
+	float Scale = ScalesData[DimIdx][SCALES_IDX];
+
+	// Computations of formulas which combine sums and divisions are done in double precision since single precision is affected by some heavy optimizations
+	// that lead to very inaccurate results. This way we obtain the same input pixels as the ORT implementations.
+
+	#if COORD_TRANS_MODE == COORD_TRANS_HALF_PIXEL
+		return DFSubtractDemote(DFDivide(DFTwoSum(Resized, 0.5f), Scale), 0.5f);
+	#elif COORD_TRANS_MODE == COORD_TRANS_HALF_PIXEL_SYMMETRIC
+		float Adjustment = asfloat(OutputTensorInfo[DimIdx][ADJUSTMENT_IDX]); // Decode adjustment to float
+		FDFScalar Center = DFDivide(InputTensorInfo[DimIdx][SIZE_IDX], 2.0f);
+		return DFSubtractDemote(DFAdd(DFMultiply(Center, DFSubtract(1.0f, Adjustment)), DFDivide(DFTwoSum(Resized, 0.5f), Scale)), 0.5f);
+	#elif COORD_TRANS_MODE == COORD_TRANS_PYTORCH_HALF_PIXEL
+		return OutputTensorInfo[DimIdx][SIZE_IDX] > 1 ? DFSubtractDemote(DFDivide(DFTwoSum(Resized, 0.5f), Scale), 0.5f) : 0.0f;
+	#elif COORD_TRANS_MODE == COORD_TRANS_ALIGN_CORNERS
+		return DFMultiplyDemote(Resized, DFDivide((InputTensorInfo[DimIdx][SIZE_IDX] - 1), (OutputTensorInfo[DimIdx][SIZE_IDX] - 1)));
+	#elif COORD_TRANS_MODE == COORD_TRANS_ASYMMETRIC
+		return DFDemote(DFDivide(Resized, Scale));
+	#elif COORD_TRANS_MODE == COORD_TRANS_TF_HALF_PIXEL_FOR_NN
+		return DFDemote(DFDivide(DFTwoSum(Resized, 0.5f), Scale));
+	#else // COORD_TRANS_MODE == COORD_TRANS_TF_CROP_AND_RESIZE
+		float Start = asfloat(InputTensorInfo[DimIdx][ROI_START_IDX]); // Decode ROI start to float
+		float End = asfloat(InputTensorInfo[DimIdx][ROI_END_IDX]); // Decode ROI end to float
+		return OutputTensorInfo[DimIdx][SIZE_IDX] > 1 ? 
+			DFAddDemote(DFMultiply(Start, (InputTensorInfo[DimIdx][SIZE_IDX] - 1)), DFDivide(DFMultiply(DFMultiply(Resized, DFSubtract(End, Start)), (InputTensorInfo[DimIdx][SIZE_IDX] - 1)), (OutputTensorInfo[DimIdx][SIZE_IDX] - 1))) 
+			: 
+			DFMultiplyDemote(DFMultiply(0.5f, DFTwoSum(Start, End)), (InputTensorInfo[DimIdx][SIZE_IDX] - 1));
+	#endif
+
+}
+
+void GetNearestInput(uint OutputIndex, out uint Iterator[NUM_DIMENSIONS])
+{
+	uint OutputIterator[NUM_DIMENSIONS];
+	
+	FromIdxToIterator(OutputIndex, OutputTensorInfo, OutputIterator);
+
+	STATIC_LOOP(DimIdx, 0, NUM_DIMENSIONS)
+	{
+		float InputCoord = TransformCoord((float) OutputIterator[DimIdx], DimIdx);
+		
+		int Dim = (int)
+			#if NEAREST_MODE == NEAREST_MODE_FLOOR
+				floor(InputCoord)
+			#elif NEAREST_MODE == NEAREST_MODE_CEIL
+				ceil(InputCoord)
+			#elif NEAREST_MODE == NEAREST_MODE_ROUND_PREFER_FLOOR
+				( InputCoord == (int) InputCoord + 0.5f ? floor(InputCoord) : round(InputCoord) )
+			#else // NEAREST_MODE == NEAREST_MODE_ROUND_PREFER_CEIL
+				( InputCoord == (int) InputCoord + 0.5f ? ceil(InputCoord) : round(InputCoord) )
+			#endif
+			;
+		
+		Dim = clamp(Dim, 0, (int) InputTensorInfo[DimIdx][SIZE_IDX] - 1);
+		
+		Iterator[DimIdx] = (uint) Dim;
+	}
+
+}
+
+#if MODE == MODE_NEAREST
+
+[numthreads(THREADGROUP_SIZE_X, 1, 1)]
+void Resize(in const uint3 DispatchThreadID : SV_DispatchThreadID)
+{
+	const uint OutputIndex = DispatchThreadID.y * ThreadCountX + DispatchThreadID.x;
+	
+	if (OutputIndex < Num)
+	{
+		uint InputIterator[NUM_DIMENSIONS];
+
+		GetNearestInput(OutputIndex, InputIterator);
+
+		uint InputIndex = FromIteratorToIdx(InputTensorInfo, InputIterator);		
+
+		Output[OutputIndex] = Input[InputIndex];
+
+	}
+}
+
+#elif MODE == MODE_LINEAR
+
+[numthreads(THREADGROUP_SIZE_X, 1, 1)]
+void Resize(in const uint3 DispatchThreadID : SV_DispatchThreadID)
+{
+	const uint OutputIndex = DispatchThreadID.y * ThreadCountX + DispatchThreadID.x;
+	
+	if (OutputIndex < Num)
+	{
+		float Weights[NUM_DIMENSIONS];
+		uint InputIterator[NUM_DIMENSIONS];
+
+		{
+			uint OutputIterator[NUM_DIMENSIONS];
+		
+			FromIdxToIterator(OutputIndex, OutputTensorInfo, OutputIterator);
+
+			STATIC_LOOP(DimIdx, 0, NUM_DIMENSIONS)
+			{
+				float InputCoord = InputTensorInfo[DimIdx][SIZE_IDX] == 1 ? 
+						(float) OutputIterator[DimIdx]
+						:
+						TransformCoord((float) OutputIterator[DimIdx], DimIdx);
+				InputCoord = max(0.0f, min(InputCoord, (float)(InputTensorInfo[DimIdx][SIZE_IDX] - 1)));
+				InputIterator[DimIdx] = (uint) floor(InputCoord);
+
+				Weights[DimIdx] = (InputIterator[DimIdx] >= InputTensorInfo[DimIdx][SIZE_IDX] - 1) ? 0.5f : frac(InputCoord);
+			}
+
+		}
+
+		float OutputValue = 0.0f;
+
+		STATIC_LOOP(PointIdx, 0, NumPointsLinear)
+		{
+			uint PointCoords[NUM_DIMENSIONS];
+			float Factor = 1.0f;
+
+			{
+
+				STATIC_LOOP(DimIdx, 0, NUM_DIMENSIONS)
+				{
+					uint PointCoord = (PointIdx & (1u << DimIdx)) >> DimIdx;
+					PointCoords[DimIdx] = InputIterator[DimIdx] + PointCoord;
+					if(PointCoords[DimIdx] >= (InputTensorInfo[DimIdx][SIZE_IDX] - 1))
+					{
+						PointCoords[DimIdx] = InputTensorInfo[DimIdx][SIZE_IDX] - 1;
+					}
+					Factor *= PointCoord == 1 ? Weights[DimIdx] : (1.0f - Weights[DimIdx]);
+
+				}
+
+			}
+
+			float PointValue = Input[FromIteratorToIdx(InputTensorInfo, PointCoords)];
+			OutputValue += PointValue * Factor;
+		}
+
+		Output[OutputIndex] = OutputValue;
+	}
+}
+
+
+#endif
\ No newline at end of file
diff --git a/Engine/Plugins/Experimental/NNERuntimeRDG/Shaders/Private/NNEHlslShaders/NNEHlslShadersScatterND.usf b/Engine/Plugins/Experimental/NNERuntimeRDG/Shaders/Private/NNEHlslShaders/NNEHlslShadersScatterND.usf
new file mode 100644
--- /dev/null
+++ b/Engine/Plugins/Experimental/NNERuntimeRDG/Shaders/Private/NNEHlslShaders/NNEHlslShadersScatterND.usf
@@ -0,0 +1,112 @@
+// Copyright Epic Games, Inc. All Rights Reserved.
+ 
+#include "/Engine/Public/Platform.ush"
+#include "NNEHlslShadersBroadcastHelper.ush"
+
+#define WORK_TYPE float
+#define BUFFER_TYPE float
+#define READ(x) x
+#define WRITE(x) x
+
+Buffer<int> InputIndices;
+Buffer<BUFFER_TYPE> InputUpdates;
+RWBuffer<BUFFER_TYPE> Output;
+uint4 DataTensorInfo[NUM_DIMENSIONS];
+
+uint Num;
+uint ThreadCountX;
+uint PartialIndexRank; // k
+uint SliceSize; // Obtained by multiplying elements in data.shape[k:r-1]
+
+#define STRIDE_IDX 0
+#define SIZE_IDX 1
+
+// Corresponds to EScatterNDReductionType defined in NNEHlslShadersScatterNDCS.h
+#define TYPE_NONE 0
+#define TYPE_ADD 1
+#define TYPE_MUL 2
+#define TYPE_MAX 3
+#define TYPE_MIN 4
+
+uint NormalizeIdx(int Idx, uint DimSize)
+{
+	return Idx < 0 ? (uint) (Idx + DimSize) : (uint) Idx;
+}
+
+// int64_t is only available from SM 6.0, therefore here we need to manually cast to int32.
+int CastInt64ToInt32(uint LSW, int MSW)
+{
+	return (MSW >> 31) == 1 ? - (int) LSW : (int) LSW;
+}
+
+#define STATIC_LOOP(Var, From, To) \
+[unroll] \
+for(uint Var = From; Var < To; ++Var)
+
+void FromIdxToIterator(const uint Idx, const uint4 TensorInfo[NUM_DIMENSIONS], out uint Iterator[NUM_DIMENSIONS])
+{
+	uint Offset = Idx;
+	STATIC_LOOP(DimIdx, 0, NUM_DIMENSIONS)
+	{
+		uint Remainder;
+		DivMod(Offset, TensorInfo[DimIdx][STRIDE_IDX], Iterator[DimIdx], Remainder);
+		Offset = Remainder;
+	}
+}
+
+uint FromIteratorToIdx(const uint4 TensorInfo[NUM_DIMENSIONS], const uint Iterator[NUM_DIMENSIONS])
+{
+	uint Offset = 0;
+	STATIC_LOOP(DimIdx, 0, NUM_DIMENSIONS)
+	{
+		Offset += Iterator[DimIdx] * TensorInfo[DimIdx][STRIDE_IDX];
+	}
+	return Offset;
+}
+
+void ApplyReduction(const uint OutputIdx, const uint UpdatesIdx)
+{
+	#if REDUCE_OPERATOR_TYPE == TYPE_NONE
+		Output[OutputIdx] = WRITE(READ(InputUpdates[UpdatesIdx]));
+	#elif REDUCE_OPERATOR_TYPE == TYPE_ADD
+		Output[OutputIdx] = WRITE(READ(Output[OutputIdx]) + READ(InputUpdates[UpdatesIdx]));
+	#elif REDUCE_OPERATOR_TYPE == TYPE_MUL
+		Output[OutputIdx] = WRITE(READ(Output[OutputIdx]) * READ(InputUpdates[UpdatesIdx]));
+	#elif REDUCE_OPERATOR_TYPE == TYPE_MAX
+		Output[OutputIdx] = WRITE(max(READ(Output[OutputIdx]), READ(InputUpdates[UpdatesIdx])));
+	#elif REDUCE_OPERATOR_TYPE == TYPE_MIN
+		Output[OutputIdx] = WRITE(min(READ(Output[OutputIdx]), READ(InputUpdates[UpdatesIdx])));
+	#endif
+}
+
+
+[numthreads(THREADGROUP_SIZE, 1, 1)]
+void ScatterND(in const uint3 DispatchThreadID : SV_DispatchThreadID)
+{
+	// See [REDACTED_PROFILE_URL]/onnx/blob/main/docs/Operators.md#ScatterND
+	// Parallelizing over updates, of shape: indices.shape[0:q-2] ++ data.shape[k:r-1]
+	const uint Idx = DispatchThreadID.y * ThreadCountX + DispatchThreadID.x;
+
+	if(Idx < Num)
+	{
+		uint ReplacementIdx = Idx / SliceSize;
+		uint ElemIdx = Idx % SliceSize;
+
+		uint OutputIterator[NUM_DIMENSIONS];
+
+		// Set the index of the element within the slice
+		FromIdxToIterator(ElemIdx, DataTensorInfo, OutputIterator);
+
+		// Then retrieve and set the index of the slice
+		for (uint DimIdx = 0; DimIdx < PartialIndexRank; ++DimIdx)
+		{
+			uint IndicesAccessIdx = ReplacementIdx * PartialIndexRank + DimIdx;
+			int RawIdx = CastInt64ToInt32( asuint(READ(InputIndices[IndicesAccessIdx * 2])), READ(InputIndices[IndicesAccessIdx * 2 + 1]) );
+			OutputIterator[DimIdx] = NormalizeIdx(RawIdx, DataTensorInfo[DimIdx][SIZE_IDX]);
+		}
+
+		uint OutputIdx = FromIteratorToIdx(DataTensorInfo, OutputIterator);
+		
+		ApplyReduction(OutputIdx, Idx);
+	}
+}
diff --git a/Engine/Plugins/Experimental/NNERuntimeRDG/Shaders/Private/NNEHlslShaders/NNEHlslShadersSlice.usf b/Engine/Plugins/Experimental/NNERuntimeRDG/Shaders/Private/NNEHlslShaders/NNEHlslShadersSlice.usf
new file mode 100644
--- /dev/null
+++ b/Engine/Plugins/Experimental/NNERuntimeRDG/Shaders/Private/NNEHlslShaders/NNEHlslShadersSlice.usf
@@ -0,0 +1,38 @@
+// Copyright Epic Games, Inc. All Rights Reserved.
+
+#include "/Engine/Public/Platform.ush"
+#include "NNEHlslShadersBroadcastHelper.ush"
+
+Buffer<float> Input;
+RWBuffer<float> Output;
+uint4 TensorInfo[NUM_DIMENSIONS];
+uint Num;
+uint ThreadCountX;
+
+#define INPUT_STRIDE 0
+#define OUTPUT_STRIDE 1
+#define INPUT_STARTS 2
+
+[numthreads(THREADGROUP_SIZE_X, 1, 1)]
+void Slice(in const uint3 DispatchThreadID : SV_DispatchThreadID)
+{
+	const uint Index = DispatchThreadID.y * ThreadCountX + DispatchThreadID.x;
+	
+	if (Index < Num)
+	{
+		uint InputIndex = 0;
+		uint Offset = Index;
+
+		for (uint dim = 0; dim < NUM_DIMENSIONS; ++dim)
+		{
+			uint OutDimIdx, R;
+			DivMod(Offset, TensorInfo[dim][OUTPUT_STRIDE], OutDimIdx, R);
+			uint InDimIdx = TensorInfo[dim][INPUT_STARTS] + OutDimIdx;
+			InputIndex += TensorInfo[dim][INPUT_STRIDE] * InDimIdx;
+			Offset = R;
+		}
+
+		float X = Input[InputIndex];
+		Output[Index] = X;
+	}
+}
diff --git a/Engine/Plugins/Experimental/NNERuntimeRDG/Shaders/Private/NNEHlslShaders/NNEHlslShadersSoftmax.usf b/Engine/Plugins/Experimental/NNERuntimeRDG/Shaders/Private/NNEHlslShaders/NNEHlslShadersSoftmax.usf
new file mode 100644
--- /dev/null
+++ b/Engine/Plugins/Experimental/NNERuntimeRDG/Shaders/Private/NNEHlslShaders/NNEHlslShadersSoftmax.usf
@@ -0,0 +1,42 @@
+// Copyright Epic Games, Inc. All Rights Reserved.
+
+#include "/Engine/Public/Platform.ush"
+
+#define WORK_TYPE float
+#define BUFFER_TYPE float
+#define READ(x) x
+#define WRITE(x) x
+
+Buffer<WORK_TYPE> Input;
+Buffer<WORK_TYPE> InputSumExp;
+RWBuffer<WORK_TYPE> Output;
+uint Num;
+uint ThreadCountX;
+uint AxisSize;
+uint AfterAxisSize;
+
+#define SOFTMAX_TYPE 0
+#define LOG_SOFTMAX_TYPE 1
+
+[numthreads(THREADGROUP_SIZE_X, 1, 1)]
+void Softmax(in const uint3 DispatchThreadID : SV_DispatchThreadID)
+{
+	const uint Index = DispatchThreadID.y * ThreadCountX + DispatchThreadID.x;
+	if (Index < Num)
+	{
+		uint SumExpIndex;
+#if SINGLE_DIMENSION == 0
+		SumExpIndex = Index / AxisSize;
+#else
+		SumExpIndex = (Index / (AfterAxisSize * AxisSize)) * AfterAxisSize + Index % AfterAxisSize;
+#endif
+		WORK_TYPE SumExp = READ(InputSumExp[SumExpIndex]);
+		WORK_TYPE Result = WRITE(exp(READ(Input[Index])) / SumExp);
+
+		#if (SOFTMAX_OPERATOR_TYPE == LOG_SOFTMAX_TYPE)
+			Result = log(Result);
+		#endif
+
+		Output[Index] = Result;
+	}
+}
\ No newline at end of file
diff --git a/Engine/Plugins/Experimental/NNERuntimeRDG/Shaders/Private/NNEHlslShaders/NNEHlslShadersSplit.usf b/Engine/Plugins/Experimental/NNERuntimeRDG/Shaders/Private/NNEHlslShaders/NNEHlslShadersSplit.usf
new file mode 100644
--- /dev/null
+++ b/Engine/Plugins/Experimental/NNERuntimeRDG/Shaders/Private/NNEHlslShaders/NNEHlslShadersSplit.usf
@@ -0,0 +1,130 @@
+// Copyright Epic Games, Inc. All Rights Reserved.
+
+#include "/Engine/Public/Platform.ush"
+#include "NNEHlslShadersBroadcastHelper.ush"
+
+Buffer<float> Input;
+
+// Buffer array is not currently available, requires a workaround. Also can't write into more than 8 output buffers.
+// RWBuffer<float> Output[MAX_NUM_SPLITS];
+#define TEMPLATE_BUFFERS(TPrefix) \
+	TPrefix##_0;\
+	TPrefix##_1;\
+	TPrefix##_2;\
+	TPrefix##_3;\
+	TPrefix##_4;\
+	TPrefix##_5;\
+	TPrefix##_6;\
+	TPrefix##_7;
+
+TEMPLATE_BUFFERS(RWBuffer<float> Output)
+
+uint4 InputTensorInfo[MAX_NUM_DIMENSIONS];
+uint4 OutputTensorInfo[MAX_NUM_SPLITS * MAX_NUM_DIMENSIONS];
+
+uint Num;
+uint ThreadCountX;
+
+#define STRIDE_IDX 0
+#define SIZE_IDX 1
+
+#define STATIC_LOOP(Var, From, To) \
+[unroll] \
+for(uint Var = From; Var < To; ++Var)
+
+void GetInputDimIterator(const uint GlobalIdx, out uint DimIterator[RANK])
+{
+	uint Offset = GlobalIdx;
+	STATIC_LOOP(DimIdx, 0, RANK)
+	{
+		uint Remainder;
+		DivMod(Offset, InputTensorInfo[DimIdx][STRIDE_IDX], DimIterator[DimIdx], Remainder);
+		Offset = Remainder;
+	}
+}
+
+struct FOutputIterator
+{
+	uint TensorIdx;
+	uint AxisIterator;
+};
+
+FOutputIterator GetOutputIterator(const uint InputAxisIterator)
+{
+	FOutputIterator OutputIterator;
+	uint SplitIndex = 0;
+	uint PrevSplitIndex = 0;
+
+	[loop]
+	for(OutputIterator.TensorIdx = 0; OutputIterator.TensorIdx < NUM_SPLITS; ++OutputIterator.TensorIdx)
+	{
+		SplitIndex += OutputTensorInfo[OutputIterator.TensorIdx * MAX_NUM_DIMENSIONS + AXIS][SIZE_IDX];
+		if(InputAxisIterator < SplitIndex)
+		{
+			OutputIterator.AxisIterator = InputAxisIterator - PrevSplitIndex;
+			return OutputIterator;
+		}
+		PrevSplitIndex = SplitIndex;
+	}
+	
+	return OutputIterator;
+}
+
+uint GetOutputGlobalIdx(const uint DimIterator[RANK], const uint OutputTensorIdx)
+{
+	uint Offset = 0;
+	STATIC_LOOP(DimIdx, 0, RANK)
+	{
+		Offset += DimIterator[DimIdx] * OutputTensorInfo[OutputTensorIdx * MAX_NUM_DIMENSIONS + DimIdx][STRIDE_IDX];
+	}
+	return Offset;
+}
+
+[numthreads(THREADGROUP_SIZE_X, 1, 1)]
+void Split(in const uint3 DispatchThreadID : SV_DispatchThreadID)
+{
+	const uint Index = DispatchThreadID.y * ThreadCountX + DispatchThreadID.x;
+	
+	if (Index < Num)
+	{
+		uint DimIterator[RANK];
+
+		GetInputDimIterator(Index, DimIterator);
+		
+		FOutputIterator OutputIterator = GetOutputIterator(DimIterator[AXIS]);
+		DimIterator[AXIS] = OutputIterator.AxisIterator;
+
+		// The following implements:
+		// Output[OutputIterator.TensorIdx][GetOutputGlobalIdx(DimIterator, OutputIterator.TensorIdx)] = Input[Index];
+
+		#define WRITE_BUFFER_ARRAY(Idx) \
+		if(OutputIterator.TensorIdx == Idx) \
+		{ \
+			Output_##Idx[GetOutputGlobalIdx(DimIterator, Idx)] = Input[Index]; \
+		}
+		#if 0 < NUM_SPLITS
+			WRITE_BUFFER_ARRAY(0)
+		#endif
+		#if 1 < NUM_SPLITS
+			WRITE_BUFFER_ARRAY(1)
+		#endif
+		#if 2 < NUM_SPLITS
+			WRITE_BUFFER_ARRAY(2)
+		#endif
+		#if 3 < NUM_SPLITS
+			WRITE_BUFFER_ARRAY(3)
+		#endif
+		#if 4 < NUM_SPLITS
+			WRITE_BUFFER_ARRAY(4)
+		#endif
+		#if 5 < NUM_SPLITS
+			WRITE_BUFFER_ARRAY(5)
+		#endif
+		#if 6 < NUM_SPLITS
+			WRITE_BUFFER_ARRAY(6)
+		#endif
+		#if 7 < NUM_SPLITS
+			WRITE_BUFFER_ARRAY(7)
+		#endif
+	}
+}
diff --git a/Engine/Plugins/Experimental/NNERuntimeRDG/Shaders/Private/NNEHlslShaders/NNEHlslShadersTranspose.usf b/Engine/Plugins/Experimental/NNERuntimeRDG/Shaders/Private/NNEHlslShaders/NNEHlslShadersTranspose.usf
new file mode 100644
--- /dev/null
+++ b/Engine/Plugins/Experimental/NNERuntimeRDG/Shaders/Private/NNEHlslShaders/NNEHlslShadersTranspose.usf
@@ -0,0 +1,36 @@
+// Copyright Epic Games, Inc. All Rights Reserved.
+
+#include "/Engine/Public/Platform.ush"
+#include "NNEHlslShadersBroadcastHelper.ush"
+
+Buffer<float> Input;
+RWBuffer<float> Output;
+uint4 TensorInfo[NUM_DIMENSIONS];
+uint Num;
+uint ThreadCountX;
+
+#define OUTPUT_STRIDE 0
+#define INPUT_STRIDE_TRANSPOSED 2
+
+[numthreads(THREADGROUP_SIZE_X, 1, 1)]
+void Transpose(in const uint3 DispatchThreadID : SV_DispatchThreadID)
+{
+	const uint Index = DispatchThreadID.y * ThreadCountX + DispatchThreadID.x;
+	
+	if (Index < Num)
+	{
+		uint InputIndex = 0;
+		uint Offset = Index;
+
+		for (uint dim = 0; dim < NUM_DIMENSIONS; ++dim)
+		{
+			uint OutDimIdx, R;
+			DivMod(Offset, TensorInfo[dim][OUTPUT_STRIDE], OutDimIdx, R);
+			InputIndex += TensorInfo[dim][INPUT_STRIDE_TRANSPOSED] * OutDimIdx;
+			Offset = R;
+		}
+
+		float X = Input[InputIndex];
+		Output[Index] = X;
+	}
+}
diff --git a/Engine/Plugins/Experimental/NNERuntimeRDG/Shaders/Private/NNEHlslShaders/NNEHlslShadersUpsample.usf b/Engine/Plugins/Experimental/NNERuntimeRDG/Shaders/Private/NNEHlslShaders/NNEHlslShadersUpsample.usf
new file mode 100644
--- /dev/null
+++ b/Engine/Plugins/Experimental/NNERuntimeRDG/Shaders/Private/NNEHlslShaders/NNEHlslShadersUpsample.usf
@@ -0,0 +1,209 @@
+// Copyright Epic Games, Inc. All Rights Reserved.
+
+#include "/Engine/Public/Platform.ush"
+#include "NNEHlslShadersBroadcastHelper.ush"
+
+Buffer<float> Input;
+RWBuffer<float> Output;
+uint4 TensorInfo[NUM_DIMENSIONS];
+uint Num;
+uint ThreadCountX;
+
+#define INPUT_STRIDE 0
+#define OUTPUT_STRIDE 1
+#define INPUT_SIZE 2
+#define SCALES_DIV 3
+
+#if MODE == 0 // Nearest
+
+[numthreads(THREADGROUP_SIZE_X, 1, 1)]
+void Upsample(in const uint3 DispatchThreadID : SV_DispatchThreadID)
+{
+	const uint OutputIndex = DispatchThreadID.y * ThreadCountX + DispatchThreadID.x;
+	
+	if (OutputIndex < Num)
+	{
+		uint InputIndex = 0;
+
+		uint CurrIndex = OutputIndex;
+		for (uint dim = 0; dim < NUM_DIMENSIONS; ++dim)
+		{
+			uint OutputCoord, OutputOffset;
+			DivMod(CurrIndex, TensorInfo[dim][OUTPUT_STRIDE], OutputCoord, OutputOffset);
+
+			uint InputCoord, InputOffset;
+			DivMod(OutputCoord, TensorInfo[dim][SCALES_DIV], InputCoord, InputOffset);
+			
+			InputIndex += TensorInfo[dim][INPUT_STRIDE] * InputCoord;
+			CurrIndex = OutputOffset;
+		}
+
+		Output[OutputIndex] = Input[InputIndex];
+	}
+}
+
+#elif MODE == 1 // Bilinear
+
+[numthreads(THREADGROUP_SIZE_X, 1, 1)]
+void Upsample(in const uint3 DispatchThreadID : SV_DispatchThreadID)
+{
+	const uint OutputIndex = DispatchThreadID.y * ThreadCountX + DispatchThreadID.x;
+	
+	if (OutputIndex < Num)
+	{
+		uint InputIndex = 0;
+
+		uint CurrIndex = OutputIndex;
+		for (uint dim = 0; dim < NUM_DIMENSIONS - 2; ++dim)
+		{
+			uint OutputCoord, OutputOffset;
+			DivMod(CurrIndex, TensorInfo[dim][OUTPUT_STRIDE], OutputCoord, OutputOffset);
+
+			uint InputCoord, InputOffset;
+			DivMod(OutputCoord, TensorInfo[dim][SCALES_DIV], InputCoord, InputOffset);
+			
+			InputIndex += TensorInfo[dim][INPUT_STRIDE] * InputCoord;
+			CurrIndex = OutputOffset;
+		}
+
+		const uint dim0 = NUM_DIMENSIONS - 2;
+		const uint dim1 = NUM_DIMENSIONS - 1;
+
+		uint OutputCoord0, OutputCoord1;
+		DivMod(CurrIndex, TensorInfo[dim0][OUTPUT_STRIDE], OutputCoord0, OutputCoord1);
+
+		uint InputCoord0, InputCoord1, InputOffset0, InputOffset1;
+		DivMod(OutputCoord0, TensorInfo[dim0][SCALES_DIV], InputCoord0, InputOffset0);
+		DivMod(OutputCoord1, TensorInfo[dim1][SCALES_DIV], InputCoord1, InputOffset1);
+
+		InputIndex += TensorInfo[dim0][INPUT_STRIDE] * InputCoord0 + InputCoord1;
+
+		const float x00 = Input[InputIndex];
+		float x10, x01, x11;
+
+		bool DimEnd0 = false;
+		if (InputCoord0 == TensorInfo[dim0][INPUT_SIZE] - 1)
+		{
+			x01 = x00;
+			DimEnd0 = true;
+		}
+		else
+		{
+			x01 = Input[InputIndex + TensorInfo[dim0][INPUT_STRIDE]];
+		}
+
+		if (InputCoord1 == TensorInfo[dim1][INPUT_SIZE] - 1)
+		{
+			x10 = x00;
+			x11 = x01;
+		}
+		else
+		{
+			x10 = Input[InputIndex + 1];
+			x11 = DimEnd0 ? x10 : Input[InputIndex + TensorInfo[dim0][INPUT_STRIDE] + 1];
+		}
+
+		float y0 = x00 + (float)InputOffset0 * (x01 - x00) / (float)TensorInfo[dim0][SCALES_DIV];
+		float y1 = x10 + (float)InputOffset0 * (x11 - x10) / (float)TensorInfo[dim0][SCALES_DIV];
+
+		Output[OutputIndex] = y0 + (float)InputOffset1 * (y1 - y0) / (float)TensorInfo[dim1][SCALES_DIV];
+	}
+}
+
+#else // Trilinear
+
+[numthreads(THREADGROUP_SIZE_X, 1, 1)]
+void Upsample(in const uint3 DispatchThreadID : SV_DispatchThreadID)
+{
+	const uint OutputIndex = DispatchThreadID.y * ThreadCountX + DispatchThreadID.x;
+	
+	if (OutputIndex < Num)
+	{
+		uint InputIndex = 0;
+
+		uint CurrIndex = OutputIndex;
+		for (uint dim = 0; dim < NUM_DIMENSIONS - 3; ++dim)
+		{
+			uint OutputCoord, OutputOffset;
+			DivMod(CurrIndex, TensorInfo[dim][OUTPUT_STRIDE], OutputCoord, OutputOffset);
+
+			uint InputCoord, InputOffset;
+			DivMod(OutputCoord, TensorInfo[dim][SCALES_DIV], InputCoord, InputOffset);
+			
+			InputIndex += TensorInfo[dim][INPUT_STRIDE] * InputCoord;
+			CurrIndex = OutputOffset;
+		}
+
+		const uint dim0 = NUM_DIMENSIONS - 3;
+		const uint dim1 = NUM_DIMENSIONS - 2;
+		const uint dim2 = NUM_DIMENSIONS - 1;
+
+		uint OutputCoord0, Mod, OutputCoord1, OutputCoord2;
+		DivMod(CurrIndex, TensorInfo[dim0][OUTPUT_STRIDE], OutputCoord0, Mod);
+		DivMod(Mod, TensorInfo[dim0][OUTPUT_STRIDE], OutputCoord1, OutputCoord2);
+
+		uint InputCoord0, InputCoord1, InputCoord2, InputOffset0, InputOffset1, InputOffset2;
+		DivMod(OutputCoord0, TensorInfo[dim0][SCALES_DIV], InputCoord0, InputOffset0);
+		DivMod(OutputCoord1, TensorInfo[dim1][SCALES_DIV], InputCoord1, InputOffset1);
+		DivMod(OutputCoord2, TensorInfo[dim2][SCALES_DIV], InputCoord2, InputOffset2);
+
+		InputIndex += TensorInfo[dim0][INPUT_STRIDE] * InputCoord0 +
+			TensorInfo[dim1][INPUT_STRIDE] * InputCoord1 +
+			InputCoord2;
+
+		const float x000 = Input[InputIndex];
+		float x010, x001, x011;
+		float x100, x110, x101, x111;
+
+		bool DimEnd0 = false;
+		if (InputCoord0 == TensorInfo[dim0][INPUT_SIZE] - 1)
+		{
+			x001 = x000;
+			DimEnd0 = true;
+		}
+		else
+		{
+			x001 = Input[InputIndex + TensorInfo[dim0][INPUT_STRIDE]];
+		}
+
+		bool DimEnd1 = false;
+		if (InputCoord1 == TensorInfo[dim1][INPUT_SIZE] - 1)
+		{
+			x010 = x000;
+			x011 = x001;
+			DimEnd1 = true;
+		}
+		else
+		{
+			x010 = Input[InputIndex + TensorInfo[dim1][INPUT_STRIDE]];
+			x011 = DimEnd0 ? x010 : Input[InputIndex + TensorInfo[dim0][INPUT_STRIDE] + TensorInfo[dim1][INPUT_STRIDE]];
+		}
+
+		if (InputCoord2 == TensorInfo[dim2][INPUT_SIZE] - 1)
+		{
+			x100 = x000;
+			x101 = x001;
+			x110 = x010;
+			x111 = x011;
+		}
+		else
+		{
+			x100 = Input[InputIndex + 1];
+			x101 = DimEnd0 ? x100 : Input[InputIndex + TensorInfo[dim0][INPUT_STRIDE] + 1];
+			x110 = DimEnd1 ? x100 : Input[InputIndex + TensorInfo[dim1][INPUT_STRIDE] + 1];
+			x111 = DimEnd0 && DimEnd1 ? x100 : DimEnd0 ? x110 : DimEnd1 ? x101 : Input[InputIndex + TensorInfo[dim0][INPUT_STRIDE] + TensorInfo[dim1][INPUT_STRIDE] + 1];
+		}
+
+		float y00 = x000 + (float)InputOffset0 * (x001 - x000) / (float)TensorInfo[dim0][SCALES_DIV];
+		float y01 = x010 + (float)InputOffset0 * (x011 - x010) / (float)TensorInfo[dim0][SCALES_DIV];
+		float y10 = x100 + (float)InputOffset0 * (x101 - x100) / (float)TensorInfo[dim0][SCALES_DIV];
+		float y11 = x110 + (float)InputOffset0 * (x111 - x110) / (float)TensorInfo[dim0][SCALES_DIV];
+
+		float z0 = y00 + (float)InputOffset1 * (y01 - y00) / (float)TensorInfo[dim1][SCALES_DIV];
+		float z1 = y10 + (float)InputOffset1 * (y11 - y10) / (float)TensorInfo[dim1][SCALES_DIV];
+
+		Output[OutputIndex] = z0 + (float)InputOffset2 * (z1 - z0) / (float)TensorInfo[dim2][SCALES_DIV];
+	}
+}
+
+#endif
\ No newline at end of file
diff --git a/Engine/Plugins/Experimental/NNERuntimeRDG/Source/NNEHlslShaders/Private/NNEHlslShadersBatchNormalizationCS.cpp b/Engine/Plugins/Experimental/NNERuntimeRDG/Source/NNEHlslShaders/Private/NNEHlslShadersBatchNormalizationCS.cpp
--- a/Engine/Plugins/Experimental/NNERuntimeRDG/Source/NNEHlslShaders/Private/NNEHlslShadersBatchNormalizationCS.cpp
+++ b/Engine/Plugins/Experimental/NNERuntimeRDG/Source/NNEHlslShaders/Private/NNEHlslShadersBatchNormalizationCS.cpp
@@ -12,5 +12,5 @@ namespace UE::NNEHlslShaders::Internal
 		OutEnvironment.SetDefine(TEXT("THREADGROUP_SIZE_X"), FBatchNormalizationConstants::NUM_GROUP_THREADS);
 	}
 
-	IMPLEMENT_GLOBAL_SHADER(TBatchNormalizationCS, "/NNE/NNEHlslShadersBatchNormalization.usf", "BatchNormalization", SF_Compute);
+	IMPLEMENT_GLOBAL_SHADER(TBatchNormalizationCS, "/NNEHlslShaders/NNEHlslShadersBatchNormalization.usf", "BatchNormalization", SF_Compute);
 } // UE::NNEHlslShaders::Internal
diff --git a/Engine/Plugins/Experimental/NNERuntimeRDG/Source/NNEHlslShaders/Private/NNEHlslShadersConvCS.cpp b/Engine/Plugins/Experimental/NNERuntimeRDG/Source/NNEHlslShaders/Private/NNEHlslShadersConvCS.cpp
--- a/Engine/Plugins/Experimental/NNERuntimeRDG/Source/NNEHlslShaders/Private/NNEHlslShadersConvCS.cpp
+++ b/Engine/Plugins/Experimental/NNERuntimeRDG/Source/NNEHlslShaders/Private/NNEHlslShadersConvCS.cpp
@@ -298,5 +298,5 @@ namespace UE::NNEHlslShaders::Internal
 		else if (FCString::Stricmp(StringVal, TEXT("VALID")) == 0) OutValue = EConvAutoPad::VALID;
 	}
 
-	IMPLEMENT_GLOBAL_SHADER(FConvCS, "/NNE/NNEHlslShadersConv.usf", "Conv", SF_Compute);
+	IMPLEMENT_GLOBAL_SHADER(FConvCS, "/NNEHlslShaders/NNEHlslShadersConv.usf", "Conv", SF_Compute);
 } // UE::NNEHlslShaders::Internal
\ No newline at end of file
diff --git a/Engine/Plugins/Experimental/NNERuntimeRDG/Source/NNEHlslShaders/Private/NNEHlslShadersConvMatmulCS.cpp b/Engine/Plugins/Experimental/NNERuntimeRDG/Source/NNEHlslShaders/Private/NNEHlslShadersConvMatmulCS.cpp
--- a/Engine/Plugins/Experimental/NNERuntimeRDG/Source/NNEHlslShaders/Private/NNEHlslShadersConvMatmulCS.cpp
+++ b/Engine/Plugins/Experimental/NNERuntimeRDG/Source/NNEHlslShaders/Private/NNEHlslShadersConvMatmulCS.cpp
@@ -22,5 +22,5 @@ namespace UE::NNEHlslShaders::Internal
 			NumBatch);
 	}
 
-	IMPLEMENT_GLOBAL_SHADER(FConvMatmulCS, "/NNE/NNEHlslShadersConvMatmul.usf", "ConvMatmul", SF_Compute);
+	IMPLEMENT_GLOBAL_SHADER(FConvMatmulCS, "/NNEHlslShaders/NNEHlslShadersConvMatmul.usf", "ConvMatmul", SF_Compute);
 } // UE::NNEHlslShaders::Internal
\ No newline at end of file
diff --git a/Engine/Plugins/Experimental/NNERuntimeRDG/Source/NNEHlslShaders/Private/NNEHlslShadersConvTransposeCS.cpp b/Engine/Plugins/Experimental/NNERuntimeRDG/Source/NNEHlslShaders/Private/NNEHlslShadersConvTransposeCS.cpp
--- a/Engine/Plugins/Experimental/NNERuntimeRDG/Source/NNEHlslShaders/Private/NNEHlslShadersConvTransposeCS.cpp
+++ b/Engine/Plugins/Experimental/NNERuntimeRDG/Source/NNEHlslShaders/Private/NNEHlslShadersConvTransposeCS.cpp
@@ -290,5 +290,5 @@ namespace UE::NNEHlslShaders::Internal
 		else if (FCString::Stricmp(StringVal, TEXT("VALID")) == 0) OutValue = EConvTransposeAutoPad::VALID;
 	}
 
-	IMPLEMENT_GLOBAL_SHADER(FConvTransposeCS, "/NNE/NNEHlslShadersConvTranspose.usf", "ConvTranspose", SF_Compute);
+	IMPLEMENT_GLOBAL_SHADER(FConvTransposeCS, "/NNEHlslShaders/NNEHlslShadersConvTranspose.usf", "ConvTranspose", SF_Compute);
 } // UE::NNEHlslShaders::Internal
\ No newline at end of file
diff --git a/Engine/Plugins/Experimental/NNERuntimeRDG/Source/NNEHlslShaders/Private/NNEHlslShadersElementWiseBinaryCS.cpp b/Engine/Plugins/Experimental/NNERuntimeRDG/Source/NNEHlslShaders/Private/NNEHlslShadersElementWiseBinaryCS.cpp
--- a/Engine/Plugins/Experimental/NNERuntimeRDG/Source/NNEHlslShaders/Private/NNEHlslShadersElementWiseBinaryCS.cpp
+++ b/Engine/Plugins/Experimental/NNERuntimeRDG/Source/NNEHlslShaders/Private/NNEHlslShadersElementWiseBinaryCS.cpp
@@ -55,5 +55,5 @@ namespace UE::NNEHlslShaders::Internal
 		return OpFunc;
 	}
 
-	IMPLEMENT_GLOBAL_SHADER(TElementWiseBinaryCS, "/NNE/NNEHlslShadersElementWiseBinary.usf", "ElementWiseBinary", SF_Compute);
+	IMPLEMENT_GLOBAL_SHADER(TElementWiseBinaryCS, "/NNEHlslShaders/NNEHlslShadersElementWiseBinary.usf", "ElementWiseBinary", SF_Compute);
 } // UE::NNEHlslShaders::Internal
diff --git a/Engine/Plugins/Experimental/NNERuntimeRDG/Source/NNEHlslShaders/Private/NNEHlslShadersElementWiseUnaryCS.cpp b/Engine/Plugins/Experimental/NNERuntimeRDG/Source/NNEHlslShaders/Private/NNEHlslShadersElementWiseUnaryCS.cpp
--- a/Engine/Plugins/Experimental/NNERuntimeRDG/Source/NNEHlslShaders/Private/NNEHlslShadersElementWiseUnaryCS.cpp
+++ b/Engine/Plugins/Experimental/NNERuntimeRDG/Source/NNEHlslShaders/Private/NNEHlslShadersElementWiseUnaryCS.cpp
@@ -94,5 +94,5 @@ namespace UE::NNEHlslShaders::Internal
 		return OpFunc;
 	}
 
-	IMPLEMENT_GLOBAL_SHADER(TElementWiseUnaryCS, "/NNE/NNEHlslShadersElementWiseUnary.usf", "ElementWiseUnary", SF_Compute);
+	IMPLEMENT_GLOBAL_SHADER(TElementWiseUnaryCS, "/NNEHlslShaders/NNEHlslShadersElementWiseUnary.usf", "ElementWiseUnary", SF_Compute);
 } // UE::NNEHlslShaders::Internal
diff --git a/Engine/Plugins/Experimental/NNERuntimeRDG/Source/NNEHlslShaders/Private/NNEHlslShadersElementWiseVariadicCS.cpp b/Engine/Plugins/Experimental/NNERuntimeRDG/Source/NNEHlslShaders/Private/NNEHlslShadersElementWiseVariadicCS.cpp
--- a/Engine/Plugins/Experimental/NNERuntimeRDG/Source/NNEHlslShaders/Private/NNEHlslShadersElementWiseVariadicCS.cpp
+++ b/Engine/Plugins/Experimental/NNERuntimeRDG/Source/NNEHlslShaders/Private/NNEHlslShadersElementWiseVariadicCS.cpp
@@ -44,5 +44,5 @@ namespace UE::NNEHlslShaders::Internal
 		return OpFunc;
 	}
 
-	IMPLEMENT_GLOBAL_SHADER(TElementWiseVariadicCS, "/NNE/NNEHlslShadersElementWiseVariadic.usf", "ElementWiseVariadic", SF_Compute);
+	IMPLEMENT_GLOBAL_SHADER(TElementWiseVariadicCS, "/NNEHlslShaders/NNEHlslShadersElementWiseVariadic.usf", "ElementWiseVariadic", SF_Compute);
 } // UE::NNEHlslShaders::Internal
diff --git a/Engine/Plugins/Experimental/NNERuntimeRDG/Source/NNEHlslShaders/Private/NNEHlslShadersGatherCS.cpp b/Engine/Plugins/Experimental/NNERuntimeRDG/Source/NNEHlslShaders/Private/NNEHlslShadersGatherCS.cpp
--- a/Engine/Plugins/Experimental/NNERuntimeRDG/Source/NNEHlslShaders/Private/NNEHlslShadersGatherCS.cpp
+++ b/Engine/Plugins/Experimental/NNERuntimeRDG/Source/NNEHlslShaders/Private/NNEHlslShadersGatherCS.cpp
@@ -81,7 +81,7 @@ namespace UE::NNEHlslShaders::Internal
 
 	//typedef TGatherCS<float, int32> FGatherFloatInt32CS;
 	//IMPLEMENT_SHADER_TYPE(template<>, TGatherCS, TEXT("/NNE/GatherOp.usf"), TEXT("main"), SF_Compute);
-	IMPLEMENT_GLOBAL_SHADER(TGatherCS, "/NNE/NNEHlslShadersGather.usf", "Gather", SF_Compute);
+	IMPLEMENT_GLOBAL_SHADER(TGatherCS, "/NNEHlslShaders/NNEHlslShadersGather.usf", "Gather", SF_Compute);
 	//template class TGatherCS<float, int32>;
 
 } // UE::NNEHlslShaders::Internal
\ No newline at end of file
diff --git a/Engine/Plugins/Experimental/NNERuntimeRDG/Source/NNEHlslShaders/Private/NNEHlslShadersGemmCS.cpp b/Engine/Plugins/Experimental/NNERuntimeRDG/Source/NNEHlslShaders/Private/NNEHlslShadersGemmCS.cpp
--- a/Engine/Plugins/Experimental/NNERuntimeRDG/Source/NNEHlslShaders/Private/NNEHlslShadersGemmCS.cpp
+++ b/Engine/Plugins/Experimental/NNERuntimeRDG/Source/NNEHlslShaders/Private/NNEHlslShadersGemmCS.cpp
@@ -193,5 +193,5 @@ namespace UE::NNEHlslShaders::Internal
 		return EGemmAlgorithm::MultiWrite1x32;
 	}
 
-	IMPLEMENT_GLOBAL_SHADER(TGemmCS, "/NNE/NNEHlslShadersGemm.usf", "Gemm", SF_Compute);
+	IMPLEMENT_GLOBAL_SHADER(TGemmCS, "/NNEHlslShaders/NNEHlslShadersGemm.usf", "Gemm", SF_Compute);
 } // UE::NNEHlslShaders::Internal
\ No newline at end of file
diff --git a/Engine/Plugins/Experimental/NNERuntimeRDG/Source/NNEHlslShaders/Private/NNEHlslShadersInstanceNormalizationCS.cpp b/Engine/Plugins/Experimental/NNERuntimeRDG/Source/NNEHlslShaders/Private/NNEHlslShadersInstanceNormalizationCS.cpp
--- a/Engine/Plugins/Experimental/NNERuntimeRDG/Source/NNEHlslShaders/Private/NNEHlslShadersInstanceNormalizationCS.cpp
+++ b/Engine/Plugins/Experimental/NNERuntimeRDG/Source/NNEHlslShaders/Private/NNEHlslShadersInstanceNormalizationCS.cpp
@@ -12,5 +12,5 @@ namespace UE::NNEHlslShaders::Internal
 		OutEnvironment.SetDefine(TEXT("THREADGROUP_SIZE_X"), FInstanceNormalizationConstants::NUM_GROUP_THREADS);
 	}
 
-	IMPLEMENT_GLOBAL_SHADER(TInstanceNormalizationCS, "/NNE/NNEHlslShadersInstanceNormalization.usf", "InstanceNormalization", SF_Compute);
+	IMPLEMENT_GLOBAL_SHADER(TInstanceNormalizationCS, "/NNEHlslShaders/NNEHlslShadersInstanceNormalization.usf", "InstanceNormalization", SF_Compute);
 } // UE::NNEHlslShaders::Internal
\ No newline at end of file
diff --git a/Engine/Plugins/Experimental/NNERuntimeRDG/Source/NNEHlslShaders/Private/NNEHlslShadersLayerNormalizationCS.cpp b/Engine/Plugins/Experimental/NNERuntimeRDG/Source/NNEHlslShaders/Private/NNEHlslShadersLayerNormalizationCS.cpp
--- a/Engine/Plugins/Experimental/NNERuntimeRDG/Source/NNEHlslShaders/Private/NNEHlslShadersLayerNormalizationCS.cpp
+++ b/Engine/Plugins/Experimental/NNERuntimeRDG/Source/NNEHlslShaders/Private/NNEHlslShadersLayerNormalizationCS.cpp
@@ -30,5 +30,5 @@ namespace UE::NNEHlslShaders::Internal
 		Parameters->Epsilon = Epsilon;
 	}
 
-	IMPLEMENT_GLOBAL_SHADER(TLayerNormalizationCS, "/NNE/NNEHlslShadersLayerNormalization.usf", "LayerNormalization", SF_Compute);
+	IMPLEMENT_GLOBAL_SHADER(TLayerNormalizationCS, "/NNEHlslShaders/NNEHlslShadersLayerNormalization.usf", "LayerNormalization", SF_Compute);
 } // UE::NNEHlslShaders::Internal
\ No newline at end of file
diff --git a/Engine/Plugins/Experimental/NNERuntimeRDG/Source/NNEHlslShaders/Private/NNEHlslShadersModule.cpp b/Engine/Plugins/Experimental/NNERuntimeRDG/Source/NNEHlslShaders/Private/NNEHlslShadersModule.cpp
--- a/Engine/Plugins/Experimental/NNERuntimeRDG/Source/NNEHlslShaders/Private/NNEHlslShadersModule.cpp
+++ b/Engine/Plugins/Experimental/NNERuntimeRDG/Source/NNEHlslShaders/Private/NNEHlslShadersModule.cpp
@@ -13,25 +13,17 @@ public:
 	// Begin IModuleInterface
 	virtual void StartupModule() override
 	{
-		FString BaseDir;
-
 		const TSharedPtr<IPlugin> Plugin = IPluginManager::Get().FindPlugin(TEXT("NNERuntimeRDG"));
-		
 		if (Plugin.IsValid())
 		{
-			BaseDir = Plugin->GetBaseDir() + TEXT("/Source/NNEHlslShaders");
+			const FString ShadersDir = FPaths::Combine(Plugin->GetBaseDir(), TEXT("Shaders/Private/NNEHlslShaders"));
+			AddShaderSourceDirectoryMapping(TEXT("/NNEHlslShaders"), ShadersDir);
 		}
 		else
 		{
-			UE_LOG(LogNNE, Warning, TEXT("Shaders directory not added. Failed to find NNE plugin"));
+			UE_LOG(LogNNE, Warning, TEXT("Shaders directory not added. Failed to find NNERuntimeRDG plugin"));
 		}
 
-		FString ModuleShaderDir = FPaths::Combine(BaseDir, TEXT("Shaders"));
-		
-		if (FPaths::DirectoryExists(ModuleShaderDir))
-		{
-			AddShaderSourceDirectoryMapping(TEXT("/NNE"), ModuleShaderDir);
-		}
 	}
 
 	virtual void ShutdownModule() override
diff --git a/Engine/Plugins/Experimental/NNERuntimeRDG/Source/NNEHlslShaders/Private/NNEHlslShadersNeuralPostProcessingCS.cpp b/Engine/Plugins/Experimental/NNERuntimeRDG/Source/NNEHlslShaders/Private/NNEHlslShadersNeuralPostProcessingCS.cpp
--- a/Engine/Plugins/Experimental/NNERuntimeRDG/Source/NNEHlslShaders/Private/NNEHlslShadersNeuralPostProcessingCS.cpp
+++ b/Engine/Plugins/Experimental/NNERuntimeRDG/Source/NNEHlslShaders/Private/NNEHlslShadersNeuralPostProcessingCS.cpp
@@ -28,8 +28,8 @@ namespace UE::NNEHlslShaders::Internal
 		OutEnvironment.SetDefine(TEXT("THREAD_GROUP_SIZE"), FNeuralPostProcessingConstants::THREAD_GROUP_SIZE);
 	}
 
-	IMPLEMENT_GLOBAL_SHADER(TNeuralPostProcessingReadInputCS, "/NNE/NNEHlslShadersNeuralPostProcessing.usf", "ReadInput", SF_Compute);
-	IMPLEMENT_GLOBAL_SHADER(TNeuralPostProcessingPreStepCS, "/NNE/NNEHlslShadersNeuralPostProcessing.usf", "PreStep", SF_Compute);
-	IMPLEMENT_GLOBAL_SHADER(TNeuralPostProcessingPostStepCS, "/NNE/NNEHlslShadersNeuralPostProcessing.usf", "PostStep", SF_Compute);
-	IMPLEMENT_GLOBAL_SHADER(TNeuralPostProcessingWriteOutputPS, "/NNE/NNEHlslShadersNeuralPostProcessing.usf", "WriteOutput", SF_Pixel);
+	IMPLEMENT_GLOBAL_SHADER(TNeuralPostProcessingReadInputCS, "/NNEHlslShaders/NNEHlslShadersNeuralPostProcessing.usf", "ReadInput", SF_Compute);
+	IMPLEMENT_GLOBAL_SHADER(TNeuralPostProcessingPreStepCS, "/NNEHlslShaders/NNEHlslShadersNeuralPostProcessing.usf", "PreStep", SF_Compute);
+	IMPLEMENT_GLOBAL_SHADER(TNeuralPostProcessingPostStepCS, "/NNEHlslShaders/NNEHlslShadersNeuralPostProcessing.usf", "PostStep", SF_Compute);
+	IMPLEMENT_GLOBAL_SHADER(TNeuralPostProcessingWriteOutputPS, "/NNEHlslShaders/NNEHlslShadersNeuralPostProcessing.usf", "WriteOutput", SF_Pixel);
 } // UE::NNEHlslShaders::Internal
\ No newline at end of file
diff --git a/Engine/Plugins/Experimental/NNERuntimeRDG/Source/NNEHlslShaders/Private/NNEHlslShadersPadCS.cpp b/Engine/Plugins/Experimental/NNERuntimeRDG/Source/NNEHlslShaders/Private/NNEHlslShadersPadCS.cpp
--- a/Engine/Plugins/Experimental/NNERuntimeRDG/Source/NNEHlslShaders/Private/NNEHlslShadersPadCS.cpp
+++ b/Engine/Plugins/Experimental/NNERuntimeRDG/Source/NNEHlslShaders/Private/NNEHlslShadersPadCS.cpp
@@ -22,5 +22,5 @@ namespace UE::NNEHlslShaders::Internal
 		else if (FCString::Stricmp(StringVal, TEXT("EDGE")) == 0) OutValue = EPadMode::EDGE;
 	}
 
-	IMPLEMENT_GLOBAL_SHADER(FPadCS, "/NNE/NNEHlslShadersPad.usf", "Pad", SF_Compute);
+	IMPLEMENT_GLOBAL_SHADER(FPadCS, "/NNEHlslShaders/NNEHlslShadersPad.usf", "Pad", SF_Compute);
 } // UE::NNEHlslShaders::Internal
\ No newline at end of file
diff --git a/Engine/Plugins/Experimental/NNERuntimeRDG/Source/NNEHlslShaders/Private/NNEHlslShadersPoolCS.cpp b/Engine/Plugins/Experimental/NNERuntimeRDG/Source/NNEHlslShaders/Private/NNEHlslShadersPoolCS.cpp
--- a/Engine/Plugins/Experimental/NNERuntimeRDG/Source/NNEHlslShaders/Private/NNEHlslShadersPoolCS.cpp
+++ b/Engine/Plugins/Experimental/NNERuntimeRDG/Source/NNEHlslShaders/Private/NNEHlslShadersPoolCS.cpp
@@ -14,5 +14,5 @@ namespace UE::NNEHlslShaders::Internal
 		FPermutationDomain PermutationVector(InParameters.PermutationId);
 	}
 
-	IMPLEMENT_GLOBAL_SHADER(FPoolCS, "/NNE/NNEHlslShadersPool.usf", "Pool", SF_Compute);
+	IMPLEMENT_GLOBAL_SHADER(FPoolCS, "/NNEHlslShaders/NNEHlslShadersPool.usf", "Pool", SF_Compute);
 } // UE::NNEHlslShaders::Internal
\ No newline at end of file
diff --git a/Engine/Plugins/Experimental/NNERuntimeRDG/Source/NNEHlslShaders/Private/NNEHlslShadersReduceCS.cpp b/Engine/Plugins/Experimental/NNERuntimeRDG/Source/NNEHlslShaders/Private/NNEHlslShadersReduceCS.cpp
--- a/Engine/Plugins/Experimental/NNERuntimeRDG/Source/NNEHlslShaders/Private/NNEHlslShadersReduceCS.cpp
+++ b/Engine/Plugins/Experimental/NNERuntimeRDG/Source/NNEHlslShaders/Private/NNEHlslShadersReduceCS.cpp
@@ -85,5 +85,5 @@ namespace UE::NNEHlslShaders::Internal
 		
 	}
 
-	IMPLEMENT_GLOBAL_SHADER(TReduceCS, "/NNE/NNEHlslShadersReduce.usf", "Reduce", SF_Compute);
+	IMPLEMENT_GLOBAL_SHADER(TReduceCS, "/NNEHlslShaders/NNEHlslShadersReduce.usf", "Reduce", SF_Compute);
 } // UE::NNEHlslShaders::Internal
diff --git a/Engine/Plugins/Experimental/NNERuntimeRDG/Source/NNEHlslShaders/Private/NNEHlslShadersResizeCS.cpp b/Engine/Plugins/Experimental/NNERuntimeRDG/Source/NNEHlslShaders/Private/NNEHlslShadersResizeCS.cpp
--- a/Engine/Plugins/Experimental/NNERuntimeRDG/Source/NNEHlslShaders/Private/NNEHlslShadersResizeCS.cpp
+++ b/Engine/Plugins/Experimental/NNERuntimeRDG/Source/NNEHlslShaders/Private/NNEHlslShadersResizeCS.cpp
@@ -69,5 +69,5 @@ namespace UE::NNEHlslShaders::Internal
         return OutValue;
 	}
 
-	IMPLEMENT_GLOBAL_SHADER(FResizeCS, "/NNE/NNEHlslShadersResize.usf", "Resize", SF_Compute);
+	IMPLEMENT_GLOBAL_SHADER(FResizeCS, "/NNEHlslShaders/NNEHlslShadersResize.usf", "Resize", SF_Compute);
 } // UE::NNEHlslShaders::Internal
\ No newline at end of file
diff --git a/Engine/Plugins/Experimental/NNERuntimeRDG/Source/NNEHlslShaders/Private/NNEHlslShadersScatterND.cpp b/Engine/Plugins/Experimental/NNERuntimeRDG/Source/NNEHlslShaders/Private/NNEHlslShadersScatterND.cpp
--- a/Engine/Plugins/Experimental/NNERuntimeRDG/Source/NNEHlslShaders/Private/NNEHlslShadersScatterND.cpp
+++ b/Engine/Plugins/Experimental/NNERuntimeRDG/Source/NNEHlslShaders/Private/NNEHlslShadersScatterND.cpp
@@ -25,5 +25,5 @@ namespace UE::NNEHlslShaders::Internal
         return OutValue;
 	}
 
-	IMPLEMENT_GLOBAL_SHADER(FScatterNDCS, "/NNE/NNEHlslShadersScatterND.usf", "ScatterND", SF_Compute);
+	IMPLEMENT_GLOBAL_SHADER(FScatterNDCS, "/NNEHlslShaders/NNEHlslShadersScatterND.usf", "ScatterND", SF_Compute);
 } // UE::NNEHlslShaders::Internal
\ No newline at end of file
diff --git a/Engine/Plugins/Experimental/NNERuntimeRDG/Source/NNEHlslShaders/Private/NNEHlslShadersSliceCS.cpp b/Engine/Plugins/Experimental/NNERuntimeRDG/Source/NNEHlslShaders/Private/NNEHlslShadersSliceCS.cpp
--- a/Engine/Plugins/Experimental/NNERuntimeRDG/Source/NNEHlslShaders/Private/NNEHlslShadersSliceCS.cpp
+++ b/Engine/Plugins/Experimental/NNERuntimeRDG/Source/NNEHlslShaders/Private/NNEHlslShadersSliceCS.cpp
@@ -14,5 +14,5 @@ namespace UE::NNEHlslShaders::Internal
 		FPermutationDomain PermutationVector(InParameters.PermutationId);
 	}
 
-	IMPLEMENT_GLOBAL_SHADER(FSliceCS, "/NNE/NNEHlslShadersSlice.usf", "Slice", SF_Compute);
+	IMPLEMENT_GLOBAL_SHADER(FSliceCS, "/NNEHlslShaders/NNEHlslShadersSlice.usf", "Slice", SF_Compute);
 } // UE::NNEHlslShaders::Internal
\ No newline at end of file
diff --git a/Engine/Plugins/Experimental/NNERuntimeRDG/Source/NNEHlslShaders/Private/NNEHlslShadersSoftmaxCS.cpp b/Engine/Plugins/Experimental/NNERuntimeRDG/Source/NNEHlslShaders/Private/NNEHlslShadersSoftmaxCS.cpp
--- a/Engine/Plugins/Experimental/NNERuntimeRDG/Source/NNEHlslShaders/Private/NNEHlslShadersSoftmaxCS.cpp
+++ b/Engine/Plugins/Experimental/NNERuntimeRDG/Source/NNEHlslShaders/Private/NNEHlslShadersSoftmaxCS.cpp
@@ -14,5 +14,5 @@ namespace UE::NNEHlslShaders::Internal
 		FPermutationDomain PermutationVector(InParameters.PermutationId);
 	}
 
-	IMPLEMENT_GLOBAL_SHADER(TSoftmaxCS, "/NNE/NNEHlslShadersSoftmax.usf", "Softmax", SF_Compute);
+	IMPLEMENT_GLOBAL_SHADER(TSoftmaxCS, "/NNEHlslShaders/NNEHlslShadersSoftmax.usf", "Softmax", SF_Compute);
 } // UE::NNEHlslShaders::Internal
diff --git a/Engine/Plugins/Experimental/NNERuntimeRDG/Source/NNEHlslShaders/Private/NNEHlslShadersSplitCS.cpp b/Engine/Plugins/Experimental/NNERuntimeRDG/Source/NNEHlslShaders/Private/NNEHlslShadersSplitCS.cpp
--- a/Engine/Plugins/Experimental/NNERuntimeRDG/Source/NNEHlslShaders/Private/NNEHlslShadersSplitCS.cpp
+++ b/Engine/Plugins/Experimental/NNERuntimeRDG/Source/NNEHlslShaders/Private/NNEHlslShadersSplitCS.cpp
@@ -23,5 +23,5 @@ namespace UE::NNEHlslShaders::Internal
 		FPermutationDomain PermutationVector(InParameters.PermutationId);
 	}
 
-	IMPLEMENT_GLOBAL_SHADER(FSplitCS, "/NNE/NNEHlslShadersSplit.usf", "Split", SF_Compute);
+	IMPLEMENT_GLOBAL_SHADER(FSplitCS, "/NNEHlslShaders/NNEHlslShadersSplit.usf", "Split", SF_Compute);
 } // UE::NNEHlslShaders::Internal
\ No newline at end of file
diff --git a/Engine/Plugins/Experimental/NNERuntimeRDG/Source/NNEHlslShaders/Private/NNEHlslShadersTransposeCS.cpp b/Engine/Plugins/Experimental/NNERuntimeRDG/Source/NNEHlslShaders/Private/NNEHlslShadersTransposeCS.cpp
--- a/Engine/Plugins/Experimental/NNERuntimeRDG/Source/NNEHlslShaders/Private/NNEHlslShadersTransposeCS.cpp
+++ b/Engine/Plugins/Experimental/NNERuntimeRDG/Source/NNEHlslShaders/Private/NNEHlslShadersTransposeCS.cpp
@@ -14,5 +14,5 @@ namespace UE::NNEHlslShaders::Internal
 		FPermutationDomain PermutationVector(InParameters.PermutationId);
 	}
 
-	IMPLEMENT_GLOBAL_SHADER(FTransposeCS, "/NNE/NNEHlslShadersTranspose.usf", "Transpose", SF_Compute);
+	IMPLEMENT_GLOBAL_SHADER(FTransposeCS, "/NNEHlslShaders/NNEHlslShadersTranspose.usf", "Transpose", SF_Compute);
 } // UE::NNEHlslShaders::Internal
\ No newline at end of file
diff --git a/Engine/Plugins/Experimental/NNERuntimeRDG/Source/NNEHlslShaders/Private/NNEHlslShadersUpsampleCS.cpp b/Engine/Plugins/Experimental/NNERuntimeRDG/Source/NNEHlslShaders/Private/NNEHlslShadersUpsampleCS.cpp
--- a/Engine/Plugins/Experimental/NNERuntimeRDG/Source/NNEHlslShaders/Private/NNEHlslShadersUpsampleCS.cpp
+++ b/Engine/Plugins/Experimental/NNERuntimeRDG/Source/NNEHlslShaders/Private/NNEHlslShadersUpsampleCS.cpp
@@ -31,5 +31,5 @@ namespace UE::NNEHlslShaders::Internal
 		FPermutationDomain PermutationVector(InParameters.PermutationId);
 	}
 
-	IMPLEMENT_GLOBAL_SHADER(FUpsampleCS, "/NNE/NNEHlslShadersUpsample.usf", "Upsample", SF_Compute);
+	IMPLEMENT_GLOBAL_SHADER(FUpsampleCS, "/NNEHlslShaders/NNEHlslShadersUpsample.usf", "Upsample", SF_Compute);
 } // UE::NNEHlslShaders::Internal
\ No newline at end of file
diff --git a/Engine/Plugins/Experimental/NNERuntimeRDG/Source/NNEHlslShaders/Shaders/NNEHlslShadersBatchNormalization.usf b/Engine/Plugins/Experimental/NNERuntimeRDG/Source/NNEHlslShaders/Shaders/NNEHlslShadersBatchNormalization.usf
deleted file mode 100644
diff --git a/Engine/Plugins/Experimental/NNERuntimeRDG/Source/NNEHlslShaders/Shaders/NNEHlslShadersBroadcastHelper.ush b/Engine/Plugins/Experimental/NNERuntimeRDG/Source/NNEHlslShaders/Shaders/NNEHlslShadersBroadcastHelper.ush
deleted file mode 100644
diff --git a/Engine/Plugins/Experimental/NNERuntimeRDG/Source/NNEHlslShaders/Shaders/NNEHlslShadersConv.usf b/Engine/Plugins/Experimental/NNERuntimeRDG/Source/NNEHlslShaders/Shaders/NNEHlslShadersConv.usf
deleted file mode 100644
diff --git a/Engine/Plugins/Experimental/NNERuntimeRDG/Source/NNEHlslShaders/Shaders/NNEHlslShadersConvMatmul.usf b/Engine/Plugins/Experimental/NNERuntimeRDG/Source/NNEHlslShaders/Shaders/NNEHlslShadersConvMatmul.usf
deleted file mode 100644
diff --git a/Engine/Plugins/Experimental/NNERuntimeRDG/Source/NNEHlslShaders/Shaders/NNEHlslShadersConvTranspose.usf b/Engine/Plugins/Experimental/NNERuntimeRDG/Source/NNEHlslShaders/Shaders/NNEHlslShadersConvTranspose.usf
deleted file mode 100644
diff --git a/Engine/Plugins/Experimental/NNERuntimeRDG/Source/NNEHlslShaders/Shaders/NNEHlslShadersElementWiseBinary.usf b/Engine/Plugins/Experimental/NNERuntimeRDG/Source/NNEHlslShaders/Shaders/NNEHlslShadersElementWiseBinary.usf
deleted file mode 100644
diff --git a/Engine/Plugins/Experimental/NNERuntimeRDG/Source/NNEHlslShaders/Shaders/NNEHlslShadersElementWiseUnary.usf b/Engine/Plugins/Experimental/NNERuntimeRDG/Source/NNEHlslShaders/Shaders/NNEHlslShadersElementWiseUnary.usf
deleted file mode 100644
diff --git a/Engine/Plugins/Experimental/NNERuntimeRDG/Source/NNEHlslShaders/Shaders/NNEHlslShadersElementWiseVariadic.usf b/Engine/Plugins/Experimental/NNERuntimeRDG/Source/NNEHlslShaders/Shaders/NNEHlslShadersElementWiseVariadic.usf
deleted file mode 100644
diff --git a/Engine/Plugins/Experimental/NNERuntimeRDG/Source/NNEHlslShaders/Shaders/NNEHlslShadersGather.usf b/Engine/Plugins/Experimental/NNERuntimeRDG/Source/NNEHlslShaders/Shaders/NNEHlslShadersGather.usf
deleted file mode 100644
diff --git a/Engine/Plugins/Experimental/NNERuntimeRDG/Source/NNEHlslShaders/Shaders/NNEHlslShadersGemm.usf b/Engine/Plugins/Experimental/NNERuntimeRDG/Source/NNEHlslShaders/Shaders/NNEHlslShadersGemm.usf
deleted file mode 100644
diff --git a/Engine/Plugins/Experimental/NNERuntimeRDG/Source/NNEHlslShaders/Shaders/NNEHlslShadersInstanceNormalization.usf b/Engine/Plugins/Experimental/NNERuntimeRDG/Source/NNEHlslShaders/Shaders/NNEHlslShadersInstanceNormalization.usf
deleted file mode 100644
diff --git a/Engine/Plugins/Experimental/NNERuntimeRDG/Source/NNEHlslShaders/Shaders/NNEHlslShadersLayerNormalization.usf b/Engine/Plugins/Experimental/NNERuntimeRDG/Source/NNEHlslShaders/Shaders/NNEHlslShadersLayerNormalization.usf
deleted file mode 100644
diff --git a/Engine/Plugins/Experimental/NNERuntimeRDG/Source/NNEHlslShaders/Shaders/NNEHlslShadersNeuralPostProcessing.usf b/Engine/Plugins/Experimental/NNERuntimeRDG/Source/NNEHlslShaders/Shaders/NNEHlslShadersNeuralPostProcessing.usf
deleted file mode 100644
diff --git a/Engine/Plugins/Experimental/NNERuntimeRDG/Source/NNEHlslShaders/Shaders/NNEHlslShadersPad.usf b/Engine/Plugins/Experimental/NNERuntimeRDG/Source/NNEHlslShaders/Shaders/NNEHlslShadersPad.usf
deleted file mode 100644
diff --git a/Engine/Plugins/Experimental/NNERuntimeRDG/Source/NNEHlslShaders/Shaders/NNEHlslShadersPool.usf b/Engine/Plugins/Experimental/NNERuntimeRDG/Source/NNEHlslShaders/Shaders/NNEHlslShadersPool.usf
deleted file mode 100644
diff --git a/Engine/Plugins/Experimental/NNERuntimeRDG/Source/NNEHlslShaders/Shaders/NNEHlslShadersReduce.usf b/Engine/Plugins/Experimental/NNERuntimeRDG/Source/NNEHlslShaders/Shaders/NNEHlslShadersReduce.usf
deleted file mode 100644
diff --git a/Engine/Plugins/Experimental/NNERuntimeRDG/Source/NNEHlslShaders/Shaders/NNEHlslShadersResize.usf b/Engine/Plugins/Experimental/NNERuntimeRDG/Source/NNEHlslShaders/Shaders/NNEHlslShadersResize.usf
deleted file mode 100644
diff --git a/Engine/Plugins/Experimental/NNERuntimeRDG/Source/NNEHlslShaders/Shaders/NNEHlslShadersScatterND.usf b/Engine/Plugins/Experimental/NNERuntimeRDG/Source/NNEHlslShaders/Shaders/NNEHlslShadersScatterND.usf
deleted file mode 100644
diff --git a/Engine/Plugins/Experimental/NNERuntimeRDG/Source/NNEHlslShaders/Shaders/NNEHlslShadersSlice.usf b/Engine/Plugins/Experimental/NNERuntimeRDG/Source/NNEHlslShaders/Shaders/NNEHlslShadersSlice.usf
deleted file mode 100644
diff --git a/Engine/Plugins/Experimental/NNERuntimeRDG/Source/NNEHlslShaders/Shaders/NNEHlslShadersSoftmax.usf b/Engine/Plugins/Experimental/NNERuntimeRDG/Source/NNEHlslShaders/Shaders/NNEHlslShadersSoftmax.usf
deleted file mode 100644
diff --git a/Engine/Plugins/Experimental/NNERuntimeRDG/Source/NNEHlslShaders/Shaders/NNEHlslShadersSplit.usf b/Engine/Plugins/Experimental/NNERuntimeRDG/Source/NNEHlslShaders/Shaders/NNEHlslShadersSplit.usf
deleted file mode 100644
diff --git a/Engine/Plugins/Experimental/NNERuntimeRDG/Source/NNEHlslShaders/Shaders/NNEHlslShadersTranspose.usf b/Engine/Plugins/Experimental/NNERuntimeRDG/Source/NNEHlslShaders/Shaders/NNEHlslShadersTranspose.usf
deleted file mode 100644
diff --git a/Engine/Plugins/Experimental/NNERuntimeRDG/Source/NNEHlslShaders/Shaders/NNEHlslShadersUpsample.usf b/Engine/Plugins/Experimental/NNERuntimeRDG/Source/NNEHlslShaders/Shaders/NNEHlslShadersUpsample.usf
deleted file mode 100644
