[NNERuntimeRDG HLSL] Implementation of the GatherElements operator

Problem: The GatherElements operator is currently not supported in HLSL.

Solution: Implement both versions (11 and 13) of the GatherElements operator.

Tested:
* Win64 Standalone tests: all
* Win64 Editor tests: all
* Console tests: smoke, GatherElementsOperator

#rnx

#virtualized

diff --git a/Engine/Plugins/Experimental/NNERuntimeRDG/Shaders/Private/NNEHlslShaders/NNEHlslShadersBroadcastHelper.ush b/Engine/Plugins/Experimental/NNERuntimeRDG/Shaders/Private/NNEHlslShaders/NNEHlslShadersBroadcastHelper.ush
--- a/Engine/Plugins/Experimental/NNERuntimeRDG/Shaders/Private/NNEHlslShaders/NNEHlslShadersBroadcastHelper.ush
+++ b/Engine/Plugins/Experimental/NNERuntimeRDG/Shaders/Private/NNEHlslShaders/NNEHlslShadersBroadcastHelper.ush
@@ -1,5 +1,7 @@
 // Copyright Epic Games, Inc. All Rights Reserved.
 
+#pragma once
+
 void DivMod(uint op1, uint op2, out uint q, out uint r)
 {
 	q = op1 / op2;
diff --git a/Engine/Plugins/Experimental/NNERuntimeRDG/Shaders/Private/NNEHlslShaders/NNEHlslShadersGather.usf b/Engine/Plugins/Experimental/NNERuntimeRDG/Shaders/Private/NNEHlslShaders/NNEHlslShadersGather.usf
--- a/Engine/Plugins/Experimental/NNERuntimeRDG/Shaders/Private/NNEHlslShaders/NNEHlslShadersGather.usf
+++ b/Engine/Plugins/Experimental/NNERuntimeRDG/Shaders/Private/NNEHlslShaders/NNEHlslShadersGather.usf
@@ -1,22 +1,7 @@
 // Copyright Epic Games, Inc. All Rights Reserved.
 
 #include "/Engine/Public/Platform.ush"
-
-#if DATA_TYPE == 0
-#define DATA_BUFFER_TYPE float
-#else
-#define DATA_BUFFER_TYPE float
-#endif
-
-#if INDICES_TYPE == 0
-#define INDICES_WORK_TYPE int
-#define INDICES_BUFFER_TYPE int
-#define INDICES_BUFFER_TO_WORK_TYPE(x) x
-#else
-#define INDICES_WORK_TYPE int
-#define INDICES_BUFFER_TYPE int
-#define INDICES_BUFFER_TO_WORK_TYPE(x) x
-#endif
+#include "NNEHlslShadersTypeHelper.ush"
 
 int Axis;
 int OutputSize;
@@ -26,9 +11,9 @@ int4 DataStride_IndicesStride_OutputStride_DataSizes[MAX_NUM_DIMENSIONS];
 float4 OneDivDataStride_OneDivIndicesStride_OneDivOutputStride[MAX_NUM_DIMENSIONS];
 
 
-Buffer<DATA_BUFFER_TYPE> Data;
-Buffer<INDICES_BUFFER_TYPE> Indices;
-RWBuffer<DATA_BUFFER_TYPE> Output;
+Buffer<float> Data;
+Buffer<int> Indices;
+RWBuffer<float> Output;
 
 uint NormalizeIdx(int Idx, uint DimSize)
 {
@@ -50,18 +35,18 @@ void Gather(in const uint3 DispatchThreadID : SV_DispatchThreadID)
 	}
 
 	// Compute the output index per dimenion
-	INDICES_WORK_TYPE OutputIndices[NUM_OUTPUT_DIMENSIONS];
-	INDICES_WORK_TYPE DispatchThreadIndex = (INDICES_WORK_TYPE)DispatchThreadID.x;
+	int OutputIndices[NUM_OUTPUT_DIMENSIONS];
+	int DispatchThreadIndex = DispatchThreadID.x;
 	UNROLL
 	for (int i = 0; i < NUM_OUTPUT_DIMENSIONS; i++)
 	{
-		OutputIndices[i] = (INDICES_WORK_TYPE)(OneDivDataStride_OneDivIndicesStride_OneDivOutputStride[i].z * (float)DispatchThreadIndex);
-		DispatchThreadIndex -= OutputIndices[i] * (INDICES_WORK_TYPE)DataStride_IndicesStride_OutputStride_DataSizes[i].z;
+		OutputIndices[i] = (int)(OneDivDataStride_OneDivIndicesStride_OneDivOutputStride[i].z * (float)DispatchThreadIndex);
+		DispatchThreadIndex -= OutputIndices[i] * DataStride_IndicesStride_OutputStride_DataSizes[i].z;
 	}
 
 	// Split the output indices into data and indices indices 
-	INDICES_WORK_TYPE DataIndices[NUM_OUTPUT_DIMENSIONS];
-	INDICES_WORK_TYPE IndicesIndices[NUM_OUTPUT_DIMENSIONS];
+	int DataIndices[NUM_OUTPUT_DIMENSIONS];
+	int IndicesIndices[NUM_OUTPUT_DIMENSIONS];
 	for (int i = 0; i < Axis; i++)
 	{
 		DataIndices[i] = OutputIndices[i];
@@ -76,27 +61,17 @@ void Gather(in const uint3 DispatchThreadID : SV_DispatchThreadID)
 	}
 
 	// Compute the flattened indices index
-	INDICES_WORK_TYPE IndicesIndex = 0;
+	int IndicesIndex = 0;
 	for (int i = 0; i < NumIndicesDimensions; i++)
 	{
 		IndicesIndex += IndicesIndices[i] * DataStride_IndicesStride_OutputStride_DataSizes[i].y;
 	}
 
-	// Look up and assign the missing data index
-	#if HAS_64BIT_INDICES == 0
-		DataIndices[Axis] = NormalizeIdx(
-			INDICES_BUFFER_TO_WORK_TYPE(Indices[IndicesIndex]), 
-			DataStride_IndicesStride_OutputStride_DataSizes[Axis].w
-		);
-	#else
-		DataIndices[Axis] = NormalizeIdx(
-			INDICES_BUFFER_TO_WORK_TYPE(CastInt64ToInt32(Indices[IndicesIndex * 2], Indices[IndicesIndex * 2 + 1])),
-			DataStride_IndicesStride_OutputStride_DataSizes[Axis].w
-		);
-	#endif
+	DataIndices[Axis] = NormalizeIdx(GetInt32IndexFromBuffer(Indices, IndicesIndex),
+								     DataStride_IndicesStride_OutputStride_DataSizes[Axis].w);
 
 	// Compute the flattened data index
-	INDICES_WORK_TYPE DataIndex = 0;
+	int DataIndex = 0;
 	for (int i = 0; i < NumDataDimensions; i++)
 	{
 		DataIndex += DataIndices[i] * DataStride_IndicesStride_OutputStride_DataSizes[i].x;
diff --git a/Engine/Plugins/Experimental/NNERuntimeRDG/Shaders/Private/NNEHlslShaders/NNEHlslShadersGatherElements.usf b/Engine/Plugins/Experimental/NNERuntimeRDG/Shaders/Private/NNEHlslShaders/NNEHlslShadersGatherElements.usf
new file mode 100644
--- /dev/null
+++ b/Engine/Plugins/Experimental/NNERuntimeRDG/Shaders/Private/NNEHlslShaders/NNEHlslShadersGatherElements.usf
@@ -0,0 +1,51 @@
+// Copyright Epic Games, Inc. All Rights Reserved.
+
+#include "/Engine/Public/Platform.ush"
+#include "NNEHlslShadersTypeHelper.ush"
+
+Buffer<float> Input;
+Buffer<int> Indices;
+RWBuffer<float> Output;
+
+int Axis;
+int AxisSize;
+int OutputSize;
+uint ThreadCountX;
+float4 OneDivOutputStrides[NUM_DIMENSIONS];
+int4 Input_OutputStrides[NUM_DIMENSIONS];
+
+[numthreads(THREADGROUP_SIZE_X, 1, 1)]
+void GatherElements(in const uint3 DispatchThreadID : SV_DispatchThreadID)
+{
+	const uint Index = DispatchThreadID.y * ThreadCountX + DispatchThreadID.x;
+
+	if(Index >= OutputSize)
+	{
+		return;
+	}
+
+	int OutputIndices[NUM_DIMENSIONS];
+	int CurrentIndex = Index;
+	UNROLL
+	for (int i = 0; i < NUM_DIMENSIONS; i++)
+	{
+		OutputIndices[i] = (int) (OneDivOutputStrides[i].x * CurrentIndex);
+		CurrentIndex -= OutputIndices[i] * Input_OutputStrides[i].y;
+	}
+
+	int AxisIndex = GetInt32IndexFromBuffer(Indices, Index);
+	if(AxisIndex < 0)
+	{
+		AxisIndex += AxisSize;
+	}
+	OutputIndices[Axis] = AxisIndex;
+	
+	int InputIndex = 0;
+	UNROLL
+	for (int i = 0; i < NUM_DIMENSIONS; i++)
+	{
+		InputIndex += OutputIndices[i] * Input_OutputStrides[i].x;
+	}
+
+	Output[Index] = Input[InputIndex];
+}
diff --git a/Engine/Plugins/Experimental/NNERuntimeRDG/Shaders/Private/NNEHlslShaders/NNEHlslShadersTypeHelper.ush b/Engine/Plugins/Experimental/NNERuntimeRDG/Shaders/Private/NNEHlslShaders/NNEHlslShadersTypeHelper.ush
new file mode 100644
--- /dev/null
+++ b/Engine/Plugins/Experimental/NNERuntimeRDG/Shaders/Private/NNEHlslShaders/NNEHlslShadersTypeHelper.ush
@@ -0,0 +1,16 @@
+// Copyright Epic Games, Inc. All Rights Reserved.
+
+#pragma once
+
+//Fetches the low 32 bits of the 64 bit two's complement number at Index and ignores the upper 32 bits. 
+//This works, assuming the value is representable as a 32 bit number, since 
+//expanding a 32 bit value to 64 bit involves copying the sign bit to the upper 32 bits and
+//therefore bit 63 == bit 31
+int GetInt32IndexFromBuffer(Buffer<int> IndexBuffer, int Index)
+{
+#if HAS_64BIT_INDICES
+	return IndexBuffer[Index * 2];
+#else
+	return IndexBuffer[Index];
+#endif
+}
diff --git a/Engine/Plugins/Experimental/NNERuntimeRDG/Source/NNEHlslShaders/Internal/NNEHlslShadersGatherElementsCS.h b/Engine/Plugins/Experimental/NNERuntimeRDG/Source/NNEHlslShaders/Internal/NNEHlslShadersGatherElementsCS.h
new file mode 100644
--- /dev/null
+++ b/Engine/Plugins/Experimental/NNERuntimeRDG/Source/NNEHlslShaders/Internal/NNEHlslShadersGatherElementsCS.h
@@ -0,0 +1,43 @@
+// Copyright Epic Games, Inc. All Rights Reserved.
+
+#pragma once
+
+#include "NNEHlslShaderBase.h"
+#include "RenderGraphUtils.h"
+#include "ShaderParameterUtils.h"
+
+namespace UE::NNEHlslShaders::Internal
+{
+	class FGatherElementsConstants
+	{
+	public:
+
+		static const int32 MAX_NUM_DIMENSIONS{ 8 };
+		static const int32 NUM_GROUP_THREADS{ 256 };
+	};
+
+	class NNEHLSLSHADERS_API FGatherElementsCS : public FHlslShaderBase
+	{
+		DECLARE_GLOBAL_SHADER(FGatherElementsCS);
+		SHADER_USE_PARAMETER_STRUCT(FGatherElementsCS, FHlslShaderBase)
+
+		class FGatherElementsDimensions : SHADER_PERMUTATION_RANGE_INT("NUM_DIMENSIONS", 1, FGatherElementsConstants::MAX_NUM_DIMENSIONS);
+		class FGatherElements64BitIndices : SHADER_PERMUTATION_BOOL("HAS_64BIT_INDICES");
+		using FPermutationDomain = TShaderPermutationDomain<FGatherElementsDimensions, FGatherElements64BitIndices>;
+
+	public:
+		BEGIN_SHADER_PARAMETER_STRUCT(FParameters, )
+			SHADER_PARAMETER_RDG_BUFFER_SRV(Buffer<float>, Input)
+			SHADER_PARAMETER_RDG_BUFFER_SRV(Buffer<int32>, Indices)
+			SHADER_PARAMETER_RDG_BUFFER_UAV(RWBuffer<float>, Output)
+			SHADER_PARAMETER_ARRAY(FVector4f, OneDivOutputStrides, [FGatherElementsConstants::MAX_NUM_DIMENSIONS])
+			SHADER_PARAMETER_ARRAY(FInt32Vector4, Input_OutputStrides, [FGatherElementsConstants::MAX_NUM_DIMENSIONS])
+			SHADER_PARAMETER(uint32, Axis)
+			SHADER_PARAMETER(uint32, AxisSize)
+			SHADER_PARAMETER(uint32, OutputSize)
+			SHADER_PARAMETER(uint32, ThreadCountX)
+		END_SHADER_PARAMETER_STRUCT()
+
+		static void ModifyCompilationEnvironment(const FGlobalShaderPermutationParameters& InParameters, FShaderCompilerEnvironment& OutEnvironment);
+	};
+} // UE::NNEHlslShaders::Internal
diff --git a/Engine/Plugins/Experimental/NNERuntimeRDG/Source/NNEHlslShaders/Private/NNEHlslShaderGatherElementsCS.cpp b/Engine/Plugins/Experimental/NNERuntimeRDG/Source/NNEHlslShaders/Private/NNEHlslShaderGatherElementsCS.cpp
new file mode 100644
--- /dev/null
+++ b/Engine/Plugins/Experimental/NNERuntimeRDG/Source/NNEHlslShaders/Private/NNEHlslShaderGatherElementsCS.cpp
@@ -0,0 +1,18 @@
+// Copyright Epic Games, Inc. All Rights Reserved.
+
+#include "NNEHlslShadersGatherElementsCS.h"
+#include "NNE.h"
+
+namespace UE::NNEHlslShaders::Internal
+{
+	void FGatherElementsCS::ModifyCompilationEnvironment(const FGlobalShaderPermutationParameters& InParameters, FShaderCompilerEnvironment& OutEnvironment)
+	{
+		FGlobalShader::ModifyCompilationEnvironment(InParameters, OutEnvironment);
+
+		OutEnvironment.SetDefine(TEXT("THREADGROUP_SIZE_X"), FGatherElementsConstants::NUM_GROUP_THREADS);
+
+		FPermutationDomain PermutationVector(InParameters.PermutationId);
+	}
+
+	IMPLEMENT_GLOBAL_SHADER(FGatherElementsCS, "/NNEHlslShaders/NNEHlslShadersGatherElements.usf", "GatherElements", SF_Compute);
+} // UE::NNEHlslShaders::Internal
diff --git a/Engine/Plugins/Experimental/NNERuntimeRDG/Source/NNERuntimeRDG/Private/Hlsl/NNERuntimeRDGGather.cpp b/Engine/Plugins/Experimental/NNERuntimeRDG/Source/NNERuntimeRDG/Private/Hlsl/NNERuntimeRDGGather.cpp
--- a/Engine/Plugins/Experimental/NNERuntimeRDG/Source/NNERuntimeRDG/Private/Hlsl/NNERuntimeRDGGather.cpp
+++ b/Engine/Plugins/Experimental/NNERuntimeRDG/Source/NNERuntimeRDG/Private/Hlsl/NNERuntimeRDGGather.cpp
@@ -129,7 +129,8 @@ namespace UE::NNERuntimeRDG::Private::Hlsl
 			TGatherCS::FParameters* Parameters = GraphBuilder.AllocParameters<TGatherCS::FParameters>();
 			TGatherCS::FillInParameters(Axis, Data, Indices, *Parameters);
 			Parameters->Data = GraphBuilder.CreateSRV(FRDGBufferSRVDesc(Data.GetBuffer(), PF_R32_FLOAT));
-			// NOTE: Indices tensor can be int64, but UE lacks support of int64 buffers. Here we use a 32-bit pixel format and in the shader we reinterpret two words as int64.
+			// NOTE: Indices tensor can be int64, but UE lacks support of int64 buffers. Here we use a 32-bit pixel format and 
+			// in the shader we discard the upper 32-bits of each value. This means that index values need to be representable by 32 bits
 			Parameters->Indices = GraphBuilder.CreateSRV(FRDGBufferSRVDesc(Indices.GetBuffer(), PF_R32_FLOAT));
 			Parameters->Output = GraphBuilder.CreateUAV(FRDGBufferUAVDesc(Output.GetBuffer(), PF_R32_FLOAT));
 
diff --git a/Engine/Plugins/Experimental/NNERuntimeRDG/Source/NNERuntimeRDG/Private/Hlsl/NNERuntimeRDGGatherElements.cpp b/Engine/Plugins/Experimental/NNERuntimeRDG/Source/NNERuntimeRDG/Private/Hlsl/NNERuntimeRDGGatherElements.cpp
new file mode 100644
--- /dev/null
+++ b/Engine/Plugins/Experimental/NNERuntimeRDG/Source/NNERuntimeRDG/Private/Hlsl/NNERuntimeRDGGatherElements.cpp
@@ -0,0 +1,164 @@
+// Copyright Epic Games, Inc. All Rights Reserved.
+
+#include "NNERuntimeRDGGatherElements.h"
+#include "NNEHlslShadersGatherElementsCS.h"
+#include "NNERuntimeRDGHlslHelper.h"
+#include "Helper/NNERuntimeRDGOperatorHelper.h"
+#include "NNETensor.h"
+#include "NNETypes.h"
+
+namespace UE::NNERuntimeRDG::Private::Hlsl
+{
+	DECLARE_GPU_STAT_NAMED(FNNEOperatorGatherElements, TEXT("NNE.Operator.Hlsl.GatherElements"));
+
+	/**
+	 * GatherElements operator implementation
+	 */
+	template<int Version>
+	class FGatherElements : public FOperatorHlsl
+	{
+	public:
+
+		FGatherElements() {}
+		virtual ~FGatherElements() = default;
+
+		int32 Rank;
+		int32 Axis;
+	public:
+
+		virtual int PrepareOutputs(TConstArrayView<NNE::Internal::FTensorRef> InputTensors, TArrayView<NNE::Internal::FTensorRef> OutputTensors) override
+		{
+			check(InputTensors.Num() == 2);
+			check(OutputTensors.Num() == 1);
+
+			OutputTensors[0]->SetShape(InputTensors[1]->GetShape());
+
+			return 0;
+		}
+
+		virtual bool Initialize(TConstArrayView<NNE::FTensorDesc> InputTensorDescs, TConstArrayView<NNE::FTensorDesc> OutputTensorDescs, const NNE::FAttributeMap& Attributes) override
+		{
+			using namespace UE::NNEHlslShaders::Internal;
+
+			check(InputTensorDescs.Num() == 2);
+			check(OutputTensorDescs.Num() == 1);
+
+			Rank = InputTensorDescs[0].GetShape().Rank();
+			Axis = Attributes.GetValueOrDefault<int32>(TEXT("axis"), 0);
+			if (Axis < -Rank || Axis >= Rank)
+			{
+				UE_LOG(LogNNE, Warning, TEXT("'axis' attribute needs to be in the range [-Rank, Rank - 1], but it is %i with a rank of %i."), Axis, Rank);
+				return false;
+			}
+			if (Axis < 0)
+			{
+				Axis += Rank;
+			}
+
+			return true;
+		}
+
+		virtual void Dispatch(FRDGBuilder& GraphBuilder, TConstArrayView<FTensorRDGRef> InputTensors, TConstArrayView<FTensorRDGRef> OutputTensors) override
+		{
+			using namespace UE::NNEHlslShaders::Internal;
+			
+			check(InputTensors.Num() == 2);
+			check(OutputTensors.Num() == 1);
+			check(InputTensors[0] != nullptr);
+			check(InputTensors[1] != nullptr);
+			check(OutputTensors[0] != nullptr);
+			const FTensorRDG& Input = *InputTensors[0];
+			const FTensorRDG& Indices = *InputTensors[1];
+			const FTensorRDG& Output = *OutputTensors[0];
+
+			FRDGBufferSRVRef InputSRV = GraphBuilder.CreateSRV(FRDGBufferSRVDesc(Input.GetBuffer(), PF_R32_FLOAT));
+			// NOTE: Indices tensor can be int64, but UE lacks support of int64 buffers. Here we use a 32-bit pixel format and 
+			// in the shader we discard the upper 32-bits of each value. This means that index values need to be representable by 32 bits
+			FRDGBufferSRVRef IndicesSRV = GraphBuilder.CreateSRV(FRDGBufferSRVDesc(Indices.GetBuffer(), PF_R32_SINT));
+			FRDGBufferUAVRef OutputUAV = GraphBuilder.CreateUAV(FRDGBufferUAVDesc(Output.GetBuffer(), PF_R32_FLOAT));
+
+			FIntVector ThreadGroupCount = ComputeElementWiseThreadGroups(Output.GetVolume(), FGatherElementsConstants::NUM_GROUP_THREADS);
+
+			// Set parameters
+			FGatherElementsCS::FParameters* Params = GraphBuilder.AllocParameters<FGatherElementsCS::FParameters>();
+			Params->Input = InputSRV;
+			Params->Indices = IndicesSRV;
+			Params->Output = OutputUAV;
+			Params->Axis = Axis;
+			Params->OutputSize = Output.GetVolume();
+			Params->ThreadCountX = ThreadGroupCount.X * FGatherElementsConstants::NUM_GROUP_THREADS;
+
+			TConstArrayView<uint32> InputShape = Input.GetShape().GetData();
+			TConstArrayView<uint32> OutputShape = Output.GetShape().GetData();
+			int32 InputStride = 1;
+			int32 OutputStride = 1;
+			for (int i = Rank - 1; i >= 0; --i)
+			{
+				Params->OneDivOutputStrides[i].X = 1.0 / (float)OutputStride;
+				Params->Input_OutputStrides[i].X = InputStride;
+				Params->Input_OutputStrides[i].Y = OutputStride;
+				OutputStride *= OutputShape[i];
+				InputStride *= InputShape[i];
+			}
+			Params->AxisSize = InputShape[Axis];
+
+			FGatherElementsCS::FPermutationDomain PermutationVector;
+			PermutationVector.Set<FGatherElementsCS::FGatherElementsDimensions>(Rank);
+			PermutationVector.Set<FGatherElementsCS::FGatherElements64BitIndices>(Indices.GetDataType() == ENNETensorDataType::Int64);
+
+			TShaderMapRef<FGatherElementsCS> ComputeShader(GetGlobalShaderMap(GMaxRHIFeatureLevel), PermutationVector);
+
+			RDG_EVENT_SCOPE(GraphBuilder, "NNE.Operator.Hlsl.GatherElements");
+			RDG_GPU_STAT_SCOPE(GraphBuilder, FNNEOperatorGatherElements);
+
+			FComputeShaderUtils::AddPass(
+				GraphBuilder,
+				RDG_EVENT_NAME("NNE.Operator.Hlsl.GatherElements.Dispatch"),
+				ERDGPassFlags::Compute | ERDGPassFlags::NeverCull,
+				ComputeShader,
+				Params,
+				ThreadGroupCount);
+		}
+	};
+
+	template<int Version>
+	bool ValidateGatherElementsOperator(const NNE::FAttributeMap& AttributeMap, TConstArrayView<ENNETensorDataType> InputTypes, TConstArrayView<NNE::FSymbolicTensorShape> InputShapes)
+	{
+		bool bIsValid = true;
+
+		FAttributeValidator AttributeValidator;
+		AttributeValidator.AddOptional(TEXT("axis"), ENNEAttributeDataType::Int32);
+		bIsValid &= AttributeValidator.Validate(AttributeMap);
+		
+		FInputValidator InputValidator;
+		InputValidator.SetTemplateCount(2);
+		InputValidator.AddSupportedType(ENNETensorDataType::Float);
+		InputValidator.AddSupportedType(ENNETensorDataType::Int64, /*TemplateIdx*/ 1);
+		InputValidator.AddSupportedType(ENNETensorDataType::Int32, /*TemplateIdx*/ 1);
+		InputValidator.AddRequired();
+		InputValidator.AddRequired(1);
+		bIsValid &= InputValidator.Validate(InputTypes);
+
+		return bIsValid;
+	}
+
+	template<int Version>
+	FOperatorHlsl* CreateGatherElementsOperator()
+	{
+		return new FGatherElements<Version>();
+	}
+
+	bool RegisterGatherElementsOperator(FOperatorRegistryHlsl& Registry)
+	{
+		// Note: support of a particular version is partial with respect to tensor data types (only the most typical ones are usually supported).
+		#define OP(Version) \
+		Registry.OpAdd({{TEXT("GatherElements"), TEXT("Onnx")}, Version}, CreateGatherElementsOperator<Version>, ValidateGatherElementsOperator<Version>);
+
+		OP(11)
+		OP(13)
+
+		#undef OP
+		return true;
+	}
+
+} // UE::NNERuntimeRDG::Private::Hlsl
diff --git a/Engine/Plugins/Experimental/NNERuntimeRDG/Source/NNERuntimeRDG/Private/Hlsl/NNERuntimeRDGGatherElements.h b/Engine/Plugins/Experimental/NNERuntimeRDG/Source/NNERuntimeRDG/Private/Hlsl/NNERuntimeRDGGatherElements.h
new file mode 100644
--- /dev/null
+++ b/Engine/Plugins/Experimental/NNERuntimeRDG/Source/NNERuntimeRDG/Private/Hlsl/NNERuntimeRDGGatherElements.h
@@ -0,0 +1,11 @@
+// Copyright Epic Games, Inc. All Rights Reserved.
+
+#pragma once
+
+#include "NNERuntimeRDGBase.h"
+#include "NNERuntimeRDGHlslOp.h"
+
+namespace UE::NNERuntimeRDG::Private::Hlsl
+{
+	bool RegisterGatherElementsOperator(FOperatorRegistryHlsl& Registry);
+} // UE::NNERuntimeRDG::Private::Hlsl
diff --git a/Engine/Plugins/Experimental/NNERuntimeRDG/Source/NNERuntimeRDG/Private/Hlsl/NNERuntimeRDGPad.cpp b/Engine/Plugins/Experimental/NNERuntimeRDG/Source/NNERuntimeRDG/Private/Hlsl/NNERuntimeRDGPad.cpp
--- a/Engine/Plugins/Experimental/NNERuntimeRDG/Source/NNERuntimeRDG/Private/Hlsl/NNERuntimeRDGPad.cpp
+++ b/Engine/Plugins/Experimental/NNERuntimeRDG/Source/NNERuntimeRDG/Private/Hlsl/NNERuntimeRDGPad.cpp
@@ -47,7 +47,7 @@ namespace UE::NNERuntimeRDG::Private::Hlsl
 					return -1;
 				}
 
-				if (InputTensors.Num() >= 3 && InputTensors[2]->GetDataType() != ENNETensorDataType::None)
+				if (InputTensors.Num() >= 3 && !InputTensors[2]->IsEmpty())
 				{
 					const NNE::Internal::FTensor& ValueTensor = *InputTensors[2];
 
@@ -66,7 +66,7 @@ namespace UE::NNERuntimeRDG::Private::Hlsl
 					Value = ValueTensor.GetPreparedData<float>()[0];
 				}
 
-				if (InputTensors.Num() >= 4 && InputTensors[3]->GetDataType() != ENNETensorDataType::None)
+				if (InputTensors.Num() >= 4 && !InputTensors[3]->IsEmpty())
 				{
 					const NNE::Internal::FTensorRef AxesTensor = InputTensors[3];
 
diff --git a/Engine/Plugins/Experimental/NNERuntimeRDG/Source/NNERuntimeRDG/Private/NNERuntimeRDGHlsl.cpp b/Engine/Plugins/Experimental/NNERuntimeRDG/Source/NNERuntimeRDG/Private/NNERuntimeRDGHlsl.cpp
--- a/Engine/Plugins/Experimental/NNERuntimeRDG/Source/NNERuntimeRDG/Private/NNERuntimeRDGHlsl.cpp
+++ b/Engine/Plugins/Experimental/NNERuntimeRDG/Source/NNERuntimeRDG/Private/NNERuntimeRDGHlsl.cpp
@@ -32,6 +32,7 @@
 #include "Hlsl/NNERuntimeRDGIdentity.h"
 #include "Hlsl/NNERuntimeRDGInstanceNormalization.h"
 #include "Hlsl/NNERuntimeRDGLayerNormalization.h"
+#include "Hlsl/NNERuntimeRDGGatherElements.h"
 #include "Hlsl/NNERuntimeRDGPad.h"
 #include "Hlsl/NNERuntimeRDGPool.h"
 #include "Hlsl/NNERuntimeRDGReduce.h"
@@ -75,6 +76,7 @@ bool UNNERuntimeRDGHlslImpl::Init()
 	RegisterIdentityOperator(*Registry);
 	RegisterInstanceNormalizationOperator(*Registry);
 	RegisterLayerNormalizationOperator(*Registry);
+	RegisterGatherElementsOperator(*Registry);
 	RegisterPadOperator(*Registry);
 	RegisterPoolOperators(*Registry);
 	RegisterReduceOperators(*Registry);
diff --git a/Engine/Source/Runtime/NNE/Internal/NNETensor.h b/Engine/Source/Runtime/NNE/Internal/NNETensor.h
--- a/Engine/Source/Runtime/NNE/Internal/NNETensor.h
+++ b/Engine/Source/Runtime/NNE/Internal/NNETensor.h
@@ -70,6 +70,13 @@ namespace UE::NNE::Internal
 			PreparedData.Append(DataAsByte);
 		}
 
+		bool IsEmpty() const
+		{
+			bool IsEmpty = GetDataType() == ENNETensorDataType::None;
+			ensureMsgf(!IsEmpty || (Shape.Rank() == 1 && Shape.GetData()[0] == 0), TEXT("Empty tensor should have a shape of [0]."));
+			return IsEmpty;
+		}
+
 		bool HasPreparedData() const
 		{
 			return !PreparedData.IsEmpty();
