NNE - Add Hlsl CumSum (partial support)

Problems:
* RDG Hlsl CumSum implementation is missing.

Solutions:
* Add partial support for Hlsl CumSum.
* Missing: support for `exclusive` and `reverse` attributes.
* Due to use of wave intrinsics and int32 atomics, there are platform support checks in place.

Tests:
* Win64 Standalone tests: all
* Win64 Editor tests: all
* Consoles tests: NNEQA_Operator_OtherTests

#rnx

#virtualized

diff --git a/Engine/Plugins/Experimental/NNERuntimeRDG/Shaders/Private/NNEHlslShaders/NNEHlslShadersCumSum.usf b/Engine/Plugins/Experimental/NNERuntimeRDG/Shaders/Private/NNEHlslShaders/NNEHlslShadersCumSum.usf
new file mode 100644
--- /dev/null
+++ b/Engine/Plugins/Experimental/NNERuntimeRDG/Shaders/Private/NNEHlslShaders/NNEHlslShadersCumSum.usf
@@ -0,0 +1,315 @@
+// Copyright Epic Games, Inc. All Rights Reserved.
+//
+// Implementation based on Single-pass Parallel Prefix Scan with Decoupled Look-back
+// Original research paper: https://research.nvidia.com/publication/2016-03_single-pass-parallel-prefix-scan-decoupled-look-back
+
+#include "/Engine/Public/Platform.ush"
+
+
+#define WORK_TYPE float
+#define BUFFER_TYPE float
+#define READ(x) x
+#define WRITE(x) x
+
+uint NumThreadGroupsPerScan; // Each group's x-axis refers to a single prefix scan (cumulative sum)
+uint NumThreadGroupsY; // Each group's y-index refers to pre-axis dimensions
+uint NumThreadGroupsZ; // Each group's z-index refers to post-axis dimensions
+
+#define STATUS_INVALID				0
+#define STATUS_AGGREGATE_AVAILABLE  1
+#define STATUS_PREFIX_AVAILABLE  	2
+
+/*
+	For each scan there is a different global partition index.
+	Each group in the same scan lock-step increments the global partition index. 
+	The value before the increment is the partition index assigned to the thread group that did the increment.
+*/
+globallycoherent RWStructuredBuffer<uint> GlobalPartitionIndex;
+/*  
+	Elements to scan are partitioned. Each thread group computes the prefix scan of a partition. 
+	Results to be propagated to the next groups are written in the partition descriptor in device memory.
+	It needs to be `globallycoherent` so that memory fences flush updates across groups.
+	There are as many descriptors as thread groups.
+*/
+struct FPartitionDescriptor
+{
+	int StatusFlag;
+	WORK_TYPE Aggregate;
+	WORK_TYPE InclusivePrefix;
+	int PadToQWord;
+};
+globallycoherent RWStructuredBuffer<FPartitionDescriptor> PartitionDescriptor;
+
+
+#ifdef INIT_SHADER
+
+uint NumInitThreadGroups;
+
+[numthreads(INIT_THREADGROUP_SIZE, 1, 1)]
+void InitCumSum(uint3 DispatchThreadID : SV_DispatchThreadID)
+{
+	const uint Increment = INIT_THREADGROUP_SIZE * NumInitThreadGroups;
+
+	for (uint Idx = DispatchThreadID.x; Idx < NumThreadGroupsPerScan * NumThreadGroupsY * NumThreadGroupsZ; Idx += Increment)
+	{
+		FPartitionDescriptor InitPD;
+
+		InitPD.StatusFlag = STATUS_INVALID;
+		InitPD.InclusivePrefix = (WORK_TYPE) 0;
+		InitPD.Aggregate = (WORK_TYPE) 0;
+		InitPD.PadToQWord = 0;
+
+		PartitionDescriptor[Idx] = InitPD;
+
+		if (Idx % NumThreadGroupsPerScan == 0)
+		{
+			GlobalPartitionIndex[Idx / NumThreadGroupsPerScan] = 0u;
+		}
+	}
+}
+
+#else // !INIT_SHADER
+
+Buffer<BUFFER_TYPE> Input;
+RWBuffer<BUFFER_TYPE> Output;
+
+#define STRIDE_IDX 0
+
+#define MIN_WAVE_LANES 4U
+#define PARTITION_SIZE (THREADGROUP_SIZE * VALUES_PER_THREAD)
+#define NUM_WAVEAGGREGATES (THREADGROUP_SIZE / WaveGetLaneCount())
+#define LAST_WAVEAGGREGATE_ID (NUM_WAVEAGGREGATES - 1)
+
+uint NumScanValues;
+uint Axis;
+uint AxisStride;
+
+/* Group shared memory variables */
+
+groupshared uint GroupPartitionIndex;
+groupshared WORK_TYPE GroupExclusivePrefix;
+groupshared WORK_TYPE LocalPrefix[PARTITION_SIZE]; // Wave-local prefixes for each element of the partition (before adding the group's exclusive prefix)
+groupshared WORK_TYPE WaveAggregate[THREADGROUP_SIZE / MIN_WAVE_LANES]; // Wave-wide aggregates (also group-local) to speed up reduction phase
+
+
+inline WORK_TYPE GetScanInput(uint IdxWithinAxis, uint ScanStartGlobalIdx)
+{
+	return Input[ ScanStartGlobalIdx + IdxWithinAxis * AxisStride ];
+}
+
+inline void SetScanOutput(uint IdxWithinAxis, WORK_TYPE Value, uint ScanStartGlobalIdx)
+{
+	Output[ ScanStartGlobalIdx + IdxWithinAxis * AxisStride ] = Value;
+}
+
+inline uint GetWaveIndex(uint GroupThreadID)
+{
+	return GroupThreadID / WaveGetLaneCount();
+}
+
+inline uint PartitionOffset(uint PartitionIndex)
+{
+	return PartitionIndex * PARTITION_SIZE;
+}
+
+inline uint WavePartitionOffset(uint GroupThreadID)
+{
+	return GetWaveIndex(GroupThreadID) * VALUES_PER_THREAD * WaveGetLaneCount();
+}
+
+/*
+	Performs a wave-local prefix scan (and reduction) of VALUES_PER_THREAD elements per thread, for each wave.
+	Scan results are stored in groupshared memory array LocalPrefix.
+	The aggregated value for each wave is stored in groupshared memory array WaveAggregate.
+*/
+inline void WaveScanReduce(uint GroupThreadID, uint PartitionIndex, uint ScanStartGlobalIdx, bool bPartial)
+{
+	const uint ShiftLaneIdx = WaveGetLaneIndex() != 0u ? (WaveGetLaneIndex() - 1u) : (WaveGetLaneCount() - 1u);
+	const uint PartialSize = NumScanValues - PartitionOffset(PartitionIndex);
+	WORK_TYPE Aggregate = 0;
+
+	[unroll]
+	for (uint Idx = WaveGetLaneIndex() + WavePartitionOffset(GroupThreadID), WordIdx = 0; WordIdx < VALUES_PER_THREAD; Idx += WaveGetLaneCount(), ++WordIdx)
+	{
+		WORK_TYPE Value = !bPartial || Idx < PartialSize ? 
+				GetScanInput(Idx + PartitionOffset(PartitionIndex), ScanStartGlobalIdx) 
+			: 	0;
+		
+		const WORK_TYPE ShiftLaneValue = WaveReadLaneAt(Value + WavePrefixSum(Value), ShiftLaneIdx); // Wave prefix sum and circular shuffle
+		LocalPrefix[Idx] = Value + (WaveGetLaneIndex() != 0u ? ShiftLaneValue : 0) + Aggregate;
+		Aggregate += WaveReadLaneAt(ShiftLaneValue, 0);
+	}
+
+	if (WaveGetLaneIndex() == 0u)
+	{
+		WaveAggregate[GetWaveIndex(GroupThreadID)] = Aggregate;
+	}
+}
+
+/*
+	Performs a group-local prefix scan of the wave aggregates in groupshared memory.
+*/
+inline void ScanWaveAggregate(uint GroupThreadID)
+{
+	const uint ScanSize = NUM_WAVEAGGREGATES;
+
+	if (GroupThreadID < WaveGetLaneCount())
+	{
+		WaveAggregate[GroupThreadID] += WavePrefixSum(WaveAggregate[GroupThreadID]);
+
+		// The following won't be needed if WaveGetLaneCount() * WaveGetLaneCount() >= THREADGROUP_SIZE (most common case).
+		for(uint AggregateOffset = WaveGetLaneCount(); AggregateOffset + GroupThreadID < ScanSize; AggregateOffset += WaveGetLaneCount())
+		{
+			GroupMemoryBarrierWithGroupSync();
+			if (WaveGetLaneIndex() == 0u)
+			{
+				WaveAggregate[AggregateOffset] += WaveAggregate[AggregateOffset - 1u];
+			}
+			GroupMemoryBarrierWithGroupSync();
+			WaveAggregate[GroupThreadID + AggregateOffset] += WavePrefixSum(WaveAggregate[GroupThreadID + AggregateOffset]);
+
+		}
+	}
+
+}
+
+
+inline void ReservePartitionIndex(uint GroupThreadID, uint ScanID)
+{
+	if (GroupThreadID == 0)
+	{
+		InterlockedAdd(GlobalPartitionIndex[ScanID], 1u, GroupPartitionIndex);
+	}
+}
+
+inline void SetPartitionInclusivePrefix(uint PDIdx, WORK_TYPE Value)
+{
+	PartitionDescriptor[PDIdx].InclusivePrefix = Value;
+	DeviceMemoryBarrier(); // To guarantee coherency avoid that status gets written before inclusive prefix does
+	PartitionDescriptor[PDIdx].StatusFlag = STATUS_PREFIX_AVAILABLE;
+}
+
+inline void SetPartitionAggregate(uint PDIdx, WORK_TYPE Value)
+{
+	PartitionDescriptor[PDIdx].Aggregate = Value;
+	DeviceMemoryBarrier(); // To guarantee coherency avoid that status gets written before aggregate does
+	PartitionDescriptor[PDIdx].StatusFlag = STATUS_AGGREGATE_AVAILABLE;
+}
+
+inline uint GetPDIdx(uint PartitionIndex, uint ScanID)
+{
+	return ScanID * NumThreadGroupsPerScan + PartitionIndex;
+}
+/*
+	Updates the current group's partition descriptor in device memory once the partition aggregate has been computed.
+*/
+inline void BroadcastGroupAggregate(uint GroupThreadID, uint PartitionIndex, uint ScanID)
+{
+	// Last wave aggregate corresponds to group (partition) aggregate
+	if (GroupThreadID == LAST_WAVEAGGREGATE_ID)
+    {
+		const uint PDIdx = GetPDIdx(PartitionIndex, ScanID);
+		if(PartitionIndex == 0)
+		{
+			SetPartitionInclusivePrefix(PDIdx, WaveAggregate[GroupThreadID]);
+		}
+		else
+		{
+			SetPartitionAggregate(PDIdx, WaveAggregate[GroupThreadID]);
+		}
+	}
+}
+/*
+	Computes the partition exclusive prefix by recursively summing up prefixes/aggregates of previous partitions in backward order.
+	Also computes and sets (broadcasts) the inclusive prefix of the partition once the exclusive prefix is available.
+*/
+inline void DecoupledLookback(uint PartitionIndex, uint ScanID)
+{
+	WORK_TYPE PreviousReduction = (WORK_TYPE) 0;
+	int LookBackIndex = (int) PartitionIndex - 1;
+
+	while (LookBackIndex >= 0)
+	{
+		const FPartitionDescriptor LookBackPD = PartitionDescriptor[GetPDIdx(LookBackIndex, ScanID)];
+
+		if (LookBackPD.StatusFlag == STATUS_PREFIX_AVAILABLE)
+		{
+			DeviceMemoryBarrier(); // Required between reading lookback's StatusFlag and InclusivePrefix for coherency
+			PreviousReduction += LookBackPD.InclusivePrefix;
+			GroupExclusivePrefix = PreviousReduction;
+			SetPartitionInclusivePrefix(
+				GetPDIdx(PartitionIndex, ScanID), 
+				PreviousReduction + WaveAggregate[LAST_WAVEAGGREGATE_ID]
+				);
+			return;
+		}
+		else if(LookBackPD.StatusFlag == STATUS_AGGREGATE_AVAILABLE)
+		{
+			DeviceMemoryBarrier(); // Required between reading lookback's StatusFlag and Aggregate for coherency
+			PreviousReduction += LookBackPD.Aggregate;
+			LookBackIndex--;
+		}
+
+	}
+
+	GroupExclusivePrefix = PreviousReduction;
+	SetPartitionInclusivePrefix(
+		GetPDIdx(PartitionIndex, ScanID), 
+		PreviousReduction + WaveAggregate[LAST_WAVEAGGREGATE_ID]
+		);
+}
+/*
+	Computes the final prefix scan by adding wave-local prefixes to wave-wise exclusive prefixes.
+*/
+void SeededWaveScan(uint GroupThreadID, uint PartitionIndex, WORK_TYPE ExclusivePrefix, uint ScanStartGlobalIdx, bool bPartial)
+{
+	const uint PartialSize = NumScanValues - PartitionOffset(PartitionIndex);
+
+	[unroll]
+	for (uint Idx = WaveGetLaneIndex() + WavePartitionOffset(GroupThreadID), WordIdx = 0; WordIdx < VALUES_PER_THREAD && (!bPartial || Idx < PartialSize); Idx += WaveGetLaneCount(), ++WordIdx)
+	{
+		SetScanOutput(Idx + PartitionOffset(PartitionIndex), LocalPrefix[Idx] + ExclusivePrefix, ScanStartGlobalIdx);
+	}
+
+}
+
+
+[numthreads(THREADGROUP_SIZE, 1, 1)]
+void CumSum(in const uint3 GroupID : SV_GroupID, in const uint3 GroupThreadIDVec : SV_GroupThreadID)
+{
+	const uint PreAxisGroupID = GroupID.y;
+	const uint PostAxisGroupID = GroupID.z;
+	const uint GroupThreadID = GroupThreadIDVec.x;
+	const uint ScanID = PreAxisGroupID * NumThreadGroupsZ + PostAxisGroupID;
+
+	// Global index to first element of the scan
+	uint ScanStartGlobalIdx = PreAxisGroupID * (NumScanValues * NumThreadGroupsZ) + PostAxisGroupID;
+
+	ReservePartitionIndex(GroupThreadID, ScanID);
+	GroupMemoryBarrierWithGroupSync();
+	const uint PartitionIndex = GroupPartitionIndex;
+
+	WaveScanReduce(GroupThreadID, PartitionIndex, ScanStartGlobalIdx, /* bPartial */ PartitionIndex == NumThreadGroupsPerScan - 1);
+	GroupMemoryBarrierWithGroupSync();
+	ScanWaveAggregate(GroupThreadID);
+
+	BroadcastGroupAggregate(GroupThreadID, PartitionIndex, ScanID);
+
+	GroupExclusivePrefix = 0.0f;
+	
+	if (PartitionIndex != 0 && GroupThreadID == 0)
+	{
+		DecoupledLookback(PartitionIndex, ScanID);
+	}
+	GroupMemoryBarrierWithGroupSync(); // Required to sync all threads in the group with respect to GroupExclusivePrefix
+
+	// Compute the wave-wise exclusive prefix
+	const WORK_TYPE ExclusivePrefix = GroupExclusivePrefix + 
+			(GroupThreadID >= WaveGetLaneCount() ? WaveAggregate[GetWaveIndex(GroupThreadID) - 1] : (WORK_TYPE) 0)
+		;
+
+	SeededWaveScan(GroupThreadID, PartitionIndex, ExclusivePrefix, ScanStartGlobalIdx, /* bPartial */ PartitionIndex == NumThreadGroupsPerScan - 1);
+
+}
+
+#endif //ifdef INIT_SHADER
\ No newline at end of file
diff --git a/Engine/Plugins/Experimental/NNERuntimeRDG/Source/NNEHlslShaders/Internal/NNEHlslShadersCumSumCS.h b/Engine/Plugins/Experimental/NNERuntimeRDG/Source/NNEHlslShaders/Internal/NNEHlslShadersCumSumCS.h
new file mode 100644
--- /dev/null
+++ b/Engine/Plugins/Experimental/NNERuntimeRDG/Source/NNEHlslShaders/Internal/NNEHlslShadersCumSumCS.h
@@ -0,0 +1,81 @@
+// Copyright Epic Games, Inc. All Rights Reserved.
+
+#pragma once
+
+#include "CoreMinimal.h"
+#include "NNEHlslShaderBase.h"
+#include "RenderGraphFwd.h"
+#include "RenderGraphUtils.h"
+#include "RHIGlobals.h"
+#include "ShaderParameterUtils.h"
+
+namespace UE::NNE::Internal { class FTensor; }
+
+namespace UE::NNEHlslShaders::Internal
+{
+
+	class FCumSumConstants
+	{
+	public:
+		static const int32 MAX_NUM_DIMENSIONS{ 8 };
+		static const int32 THREADGROUP_SIZE{ 256 };
+		static const int32 VALUES_PER_THREAD{ 12 };
+		static const int32 INIT_THREADGROUP_SIZE{ 768 };
+	};
+
+	struct FPartitionDescriptor
+	{
+		int32 StatusFlag;
+		float Aggregate;
+		float InclusivePrefix;
+		int32 PadToQWord;
+	};
+
+	class NNEHLSLSHADERS_API TInitCumSumCS : public FHlslShaderBase
+	{
+		DECLARE_GLOBAL_SHADER(TInitCumSumCS);
+		SHADER_USE_PARAMETER_STRUCT(TInitCumSumCS, FHlslShaderBase)
+
+		using FPermutationDomain = TShaderPermutationDomain<>;
+
+	public:
+
+		BEGIN_SHADER_PARAMETER_STRUCT(FParameters, )
+			SHADER_PARAMETER(uint32, NumThreadGroupsPerScan)
+			SHADER_PARAMETER(uint32, NumThreadGroupsY)
+			SHADER_PARAMETER(uint32, NumThreadGroupsZ)
+			SHADER_PARAMETER(uint32, NumInitThreadGroups)
+			SHADER_PARAMETER_RDG_BUFFER_UAV(RWStructuredBuffer<uint32>, GlobalPartitionIndex)
+			SHADER_PARAMETER_RDG_BUFFER_UAV(RWStructuredBuffer<FPartitionDescriptor>, PartitionDescriptor)
+		END_SHADER_PARAMETER_STRUCT()
+
+		static void ModifyCompilationEnvironment(const FGlobalShaderPermutationParameters& InParameters, FShaderCompilerEnvironment& OutEnvironment);
+	};
+
+	class NNEHLSLSHADERS_API TCumSumCS : public FHlslShaderBase
+	{
+		DECLARE_GLOBAL_SHADER(TCumSumCS);
+		SHADER_USE_PARAMETER_STRUCT(TCumSumCS, FHlslShaderBase)
+
+		using FPermutationDomain = TShaderPermutationDomain<>;
+
+	public:
+
+		BEGIN_SHADER_PARAMETER_STRUCT(FParameters, )
+			SHADER_PARAMETER(uint32, NumThreadGroupsPerScan)
+			SHADER_PARAMETER(uint32, NumThreadGroupsY)
+			SHADER_PARAMETER(uint32, NumThreadGroupsZ)
+			SHADER_PARAMETER(uint32, NumScanValues)
+			SHADER_PARAMETER(uint32, Axis)
+			SHADER_PARAMETER(uint32, AxisStride)
+			SHADER_PARAMETER_RDG_BUFFER_SRV(Buffer<float>, Input)
+			SHADER_PARAMETER_RDG_BUFFER_UAV(RWBuffer<float>, Output)
+			SHADER_PARAMETER_ARRAY(FUintVector4, TensorInfo, [FCumSumConstants::MAX_NUM_DIMENSIONS])
+			SHADER_PARAMETER_RDG_BUFFER_UAV(RWStructuredBuffer<uint32>, GlobalPartitionIndex)
+			SHADER_PARAMETER_RDG_BUFFER_UAV(RWStructuredBuffer<FPartitionDescriptor>, PartitionDescriptor)
+		END_SHADER_PARAMETER_STRUCT()
+
+		static void ModifyCompilationEnvironment(const FGlobalShaderPermutationParameters& InParameters, FShaderCompilerEnvironment& OutEnvironment);
+		static bool ShouldCompilePermutation(const FGlobalShaderPermutationParameters& InParameters);
+	};
+} // UE::NNEHlslShaders::Internal
\ No newline at end of file
diff --git a/Engine/Plugins/Experimental/NNERuntimeRDG/Source/NNEHlslShaders/Private/NNEHlslShadersCumSumCS.cpp b/Engine/Plugins/Experimental/NNERuntimeRDG/Source/NNEHlslShaders/Private/NNEHlslShadersCumSumCS.cpp
new file mode 100644
--- /dev/null
+++ b/Engine/Plugins/Experimental/NNERuntimeRDG/Source/NNEHlslShaders/Private/NNEHlslShadersCumSumCS.cpp
@@ -0,0 +1,39 @@
+// Copyright Epic Games, Inc. All Rights Reserved.
+
+#include "NNEHlslShadersCumSumCS.h"
+#include "DataDrivenShaderPlatformInfo.h"
+#include "NNE.h"
+#include "RenderGraphBuilder.h"
+
+namespace UE::NNEHlslShaders::Internal
+{
+	void TInitCumSumCS::ModifyCompilationEnvironment(const FGlobalShaderPermutationParameters& InParameters, FShaderCompilerEnvironment& OutEnvironment)
+	{
+		FGlobalShader::ModifyCompilationEnvironment(InParameters, OutEnvironment);
+		OutEnvironment.SetDefine(TEXT("INIT_THREADGROUP_SIZE"), FCumSumConstants::INIT_THREADGROUP_SIZE);
+		OutEnvironment.SetDefine(TEXT("INIT_SHADER"), 1);
+	}
+
+	void TCumSumCS::ModifyCompilationEnvironment(const FGlobalShaderPermutationParameters& InParameters, FShaderCompilerEnvironment& OutEnvironment)
+	{
+		FGlobalShader::ModifyCompilationEnvironment(InParameters, OutEnvironment);
+
+		OutEnvironment.SetDefine(TEXT("THREADGROUP_SIZE"), FCumSumConstants::THREADGROUP_SIZE);
+		OutEnvironment.SetDefine(TEXT("VALUES_PER_THREAD"), FCumSumConstants::VALUES_PER_THREAD);
+
+		FPermutationDomain PermutationVector(InParameters.PermutationId);
+	}
+
+	bool TCumSumCS::ShouldCompilePermutation(const FGlobalShaderPermutationParameters& InParameters)
+	{
+		//NOTE: platform compatibility checks should be rather done in a NNERuntimeRDG-wise place and called from here
+
+		const ERHIFeatureSupport WaveOpsSupport = FDataDrivenShaderPlatformInfo::GetSupportsWaveOperations(InParameters.Platform);
+		const bool bWaveOpsAvailable = WaveOpsSupport == ERHIFeatureSupport::RuntimeDependent || WaveOpsSupport == ERHIFeatureSupport::RuntimeGuaranteed;
+
+		return bWaveOpsAvailable;
+	}
+
+	IMPLEMENT_GLOBAL_SHADER(TInitCumSumCS, "/NNEHlslShaders/NNEHlslShadersCumSum.usf", "InitCumSum", SF_Compute);
+	IMPLEMENT_GLOBAL_SHADER(TCumSumCS, "/NNEHlslShaders/NNEHlslShadersCumSum.usf", "CumSum", SF_Compute);
+} // UE::NNEHlslShaders::Internal
diff --git a/Engine/Plugins/Experimental/NNERuntimeRDG/Source/NNERuntimeRDG/Private/Hlsl/NNERuntimeRDGCumSum.cpp b/Engine/Plugins/Experimental/NNERuntimeRDG/Source/NNERuntimeRDG/Private/Hlsl/NNERuntimeRDGCumSum.cpp
new file mode 100644
--- /dev/null
+++ b/Engine/Plugins/Experimental/NNERuntimeRDG/Source/NNERuntimeRDG/Private/Hlsl/NNERuntimeRDGCumSum.cpp
@@ -0,0 +1,296 @@
+// Copyright Epic Games, Inc. All Rights Reserved.
+
+#include "NNEHlslShadersCumSumCS.h"
+#include "Helper/NNERuntimeRDGOperatorHelper.h"
+#include "NNERuntimeRDGCumSum.h"
+#include "NNERuntimeRDGHlslHelper.h"
+#include "NNETensor.h"
+#include "NNETypes.h"
+#include "RenderGraphUtils.h"
+
+namespace UE::NNERuntimeRDG::Private::Hlsl
+{
+	DECLARE_GPU_STAT_NAMED(FNNEOperatorCumSum, TEXT("NNE.Operator.Hlsl.CumSum"));
+
+	/**
+	 * CumSum operator implementation
+	 */
+	template<int Version>
+	class FCumSumOperator : public FOperatorHlsl
+	{
+
+	public:
+
+		FCumSumOperator() = default;
+		virtual ~FCumSumOperator() = default;
+
+	private:
+
+		int32 Axis;
+		FIntVector ThreadGroupCount;
+
+	public:
+
+		virtual int PrepareOutputs(TConstArrayView<NNE::Internal::FTensorRef> InputTensors, TArrayView<NNE::Internal::FTensorRef> OutputTensors) override
+		{
+			check(InputTensors.Num() == 2);
+			check(OutputTensors.Num() == 1);
+
+			const NNE::Internal::FTensorRef AxisTensor = InputTensors[1];
+			if(!AxisTensor->HasPreparedData())
+			{
+				UE_LOG(LogNNE, Warning, TEXT("CumSum tensor `axis` (name: %s) must be CPU constant."), *InputTensors[1]->GetName());
+				return -1;
+			}
+			
+			if(AxisTensor->GetVolume() != 1)
+			{
+				UE_LOG(LogNNE, Warning, TEXT("CumSum tensor `axis` (name: %s) must be 0-D."), *InputTensors[1]->GetName());
+				return -1;
+			}
+
+			TArray<int32, TInlineAllocator<NNE::FTensorShape::MaxRank>> TempAxisArray;
+			OperatorHelper::GetInt32ArrayFromConstTensor(TempAxisArray, AxisTensor);
+			check(TempAxisArray.Num() == 1);
+			Axis = TempAxisArray[0];
+
+			const NNE::FTensorShape& InputShape = InputTensors[0]->GetShape();
+			
+			const int32 InputRank = InputShape.Rank();
+
+			if (Axis > InputRank || Axis  < -InputRank)
+			{
+				UE_LOG(LogNNE, Warning, TEXT("CumSum tensor 'axis' should contain a value in the range [-r,r] with r being the rank of the input (name: %s) however got %d while rank is %d."), *InputTensors[0]->GetName(), Axis, InputRank);
+				return -1;
+			}
+
+			if (Axis  < 0)
+			{
+				Axis += InputRank;
+			}
+
+			auto ComputeThreadGroupCount = [Axis=Axis](const NNE::Internal::FTensor& InputTensor) -> FIntVector
+			{
+				FIntVector OutThreadGroupCount;
+				TConstArrayView<uint32> InputShape = InputTensor.GetShape().GetData();
+
+				const uint32 PrefixSumPartitionSize = (uint32) (NNEHlslShaders::Internal::FCumSumConstants::THREADGROUP_SIZE * NNEHlslShaders::Internal::FCumSumConstants::VALUES_PER_THREAD);
+
+				OutThreadGroupCount.X = FMath::DivideAndRoundUp(InputShape[Axis], PrefixSumPartitionSize);
+
+				if(OutThreadGroupCount.X > GRHIMaxDispatchThreadGroupsPerDimension.X)
+				{
+					UE_LOG(LogNNE, Warning, TEXT("CumSum input tensor (name: %s) has axis dimension greater than %d. This is not supported."), 
+						*InputTensor.GetName(), GRHIMaxDispatchThreadGroupsPerDimension.X);
+					return FIntVector{};
+				}
+				
+				int32 NumElemBeforeAxis = 1;
+				int32 NumElemAfterAxis = 1;
+				for (int32 DimIdx = 0; DimIdx < Axis; ++DimIdx)
+				{
+					NumElemBeforeAxis *= InputShape[DimIdx];
+				}
+				for (int32 DimIdx = Axis + 1; DimIdx < InputShape.Num(); ++DimIdx)
+				{
+					NumElemAfterAxis *= InputShape[DimIdx];
+				}
+				
+				OutThreadGroupCount.Y = NumElemBeforeAxis;
+				if(OutThreadGroupCount.Y > GRHIMaxDispatchThreadGroupsPerDimension.Y)
+				{
+					UE_LOG(LogNNE, Warning, TEXT("CumSum input tensor (name: %s) has number of elements before axis greater than %d. This is not supported."), 
+						*InputTensor.GetName(), GRHIMaxDispatchThreadGroupsPerDimension.X);
+					return FIntVector{};
+				}
+				OutThreadGroupCount.Z = NumElemAfterAxis;
+				if(OutThreadGroupCount.Z > GRHIMaxDispatchThreadGroupsPerDimension.Z)
+				{
+					UE_LOG(LogNNE, Warning, TEXT("CumSum input tensor (name: %s) has number of elements after axis greater than %d. This is not supported."), 
+						*InputTensor.GetName(), GRHIMaxDispatchThreadGroupsPerDimension.Z);
+					return FIntVector{};
+				}
+
+				return OutThreadGroupCount;
+			};
+
+			ThreadGroupCount = ComputeThreadGroupCount(*InputTensors[0]);
+			if(ThreadGroupCount == FIntVector{})
+			{
+				return -1;
+			}
+
+			OutputTensors[0]->SetShape(InputShape);
+
+			return 0;
+		};
+
+		virtual bool Initialize(TConstArrayView<NNE::FTensorDesc> InputTensorDescs, TConstArrayView<NNE::FTensorDesc> OutputTensorDescs, const NNE::FAttributeMap& Attributes) override
+		{
+			check(InputTensorDescs.Num() == 2);
+			check(OutputTensorDescs.Num() == 1);
+
+			//NOTE: hardware compatibility check should rather be done in a NNERuntimeRDG-wise place
+			if(!GRHISupportsWaveOperations)
+			{
+				UE_LOG(LogNNE, Warning, TEXT("Hlsl CumSum not supported on current hardware."));
+				return false;
+			}
+
+			const int32 InputRank = InputTensorDescs[0].GetShape().Rank();
+			
+			int32 Exclusive = Attributes.GetValueOrDefault<int32>(TEXT("exclusive"), 0);
+			if(Exclusive == 1)
+			{
+				UE_LOG(LogNNE, Warning, TEXT("CumSum attribute `exclusive` not yet supported."));
+				return false;
+			}
+
+			int32 Reverse = Attributes.GetValueOrDefault<int32>(TEXT("reverse"), 0);
+			if(Reverse == 1)
+			{
+				UE_LOG(LogNNE, Warning, TEXT("CumSum attribute `reverse` not yet supported."));
+				return false;
+			}
+
+			return true;
+		}
+
+		virtual void Dispatch(FRDGBuilder& GraphBuilder, TConstArrayView<FTensorRDGRef> InputTensors, TConstArrayView<FTensorRDGRef> OutputTensors) override
+		{
+			using namespace UE::NNEHlslShaders::Internal;
+
+			check(InputTensors.Num() == 2);
+			check(OutputTensors.Num() == 1);
+			check(InputTensors[0] != FTensorRDGRef{});
+			check(OutputTensors[0] != FTensorRDGRef{});
+			check(ThreadGroupCount != FIntVector{});
+
+			RDG_EVENT_SCOPE(GraphBuilder, "NNE.Operator.Hlsl.CumSum");
+			RDG_GPU_STAT_SCOPE(GraphBuilder, FNNEOperatorCumSum);
+
+			const FTensorRDG& Input = *InputTensors[0];
+			const FTensorRDG& Output = *OutputTensors[0];
+
+			const uint32 ThreadGroupCountTotal = ThreadGroupCount.X * ThreadGroupCount.Y * ThreadGroupCount.Z;
+			const uint32 NumParallelScans = ThreadGroupCount.Y * ThreadGroupCount.Z;
+
+			const FRDGBufferDesc GPIBufferDesc = FRDGBufferDesc::CreateStructuredDesc(sizeof(uint32), NumParallelScans);
+			FRDGBufferRef GPIBuffer = GraphBuilder.CreateBuffer(GPIBufferDesc, TEXT("NNE.Operator.Hlsl.CumSum.GPIBuffer"));
+
+			const FRDGBufferDesc PartitionDescriptorBufferDesc = FRDGBufferDesc::CreateStructuredDesc(sizeof(FPartitionDescriptor), ThreadGroupCountTotal);
+			FRDGBufferRef PartitionDescriptorBuffer = GraphBuilder.CreateBuffer(PartitionDescriptorBufferDesc, TEXT("NNE.Operator.Hlsl.CumSum.PartitionDescriptorBuffer"));
+
+			{
+				TInitCumSumCS::FParameters* Params = GraphBuilder.AllocParameters<TInitCumSumCS::FParameters>();
+				Params->GlobalPartitionIndex = GraphBuilder.CreateUAV(GPIBuffer);
+				Params->PartitionDescriptor = GraphBuilder.CreateUAV(PartitionDescriptorBuffer);
+				Params->NumThreadGroupsPerScan = ThreadGroupCount.X;
+				Params->NumThreadGroupsY = ThreadGroupCount.Y;
+				Params->NumThreadGroupsZ = ThreadGroupCount.Z;
+
+				const FIntVector InitThreadGroupCount = ComputeElementWiseThreadGroups(ThreadGroupCountTotal, FCumSumConstants::INIT_THREADGROUP_SIZE);
+				Params->NumInitThreadGroups = InitThreadGroupCount.X;
+
+				TInitCumSumCS::FPermutationDomain PermutationVector;
+
+				TShaderMapRef<TInitCumSumCS> ComputeShader(GetGlobalShaderMap(GMaxRHIFeatureLevel), PermutationVector);
+
+				FComputeShaderUtils::AddPass(
+					GraphBuilder,
+					RDG_EVENT_NAME("NNE.Operator.Hlsl.InitCumSum.Dispatch"),
+					ERDGPassFlags::Compute | ERDGPassFlags::NeverCull,
+					ComputeShader,
+					Params,
+					InitThreadGroupCount);
+				
+			}
+			
+			{
+				const FRDGBufferSRVRef InputSRV = GraphBuilder.CreateSRV(FRDGBufferSRVDesc(Input.GetBuffer(), PF_R32_FLOAT));
+				const FRDGBufferUAVRef OutputUAV = GraphBuilder.CreateUAV(FRDGBufferUAVDesc(Output.GetBuffer(), PF_R32_FLOAT));
+				
+				// Set parameters
+				TCumSumCS::FParameters* Params = GraphBuilder.AllocParameters<TCumSumCS::FParameters>();
+				Params->Input = InputSRV;
+				Params->Output = OutputUAV;
+				Params->GlobalPartitionIndex = GraphBuilder.CreateUAV(GPIBuffer);
+				Params->PartitionDescriptor = GraphBuilder.CreateUAV(PartitionDescriptorBuffer);
+				
+				Params->NumThreadGroupsPerScan = ThreadGroupCount.X;
+				Params->NumThreadGroupsY = ThreadGroupCount.Y;
+				Params->NumThreadGroupsZ = ThreadGroupCount.Z;
+				Params->NumScanValues = Input.GetShape().GetData()[Axis];
+				Params->Axis = (uint32) Axis;
+				// Compute Axis' stride
+				{
+					uint32 Stride = 1;
+					for (int32 Idx = Input.GetShape().Rank() - 1; Idx > Axis; --Idx)
+					{
+						Stride *= Input.GetShape().GetData()[Idx];
+					}
+					Params->AxisStride = Stride;
+				}
+
+				TCumSumCS::FPermutationDomain PermutationVector;
+
+				TShaderMapRef<TCumSumCS> ComputeShader(GetGlobalShaderMap(GMaxRHIFeatureLevel), PermutationVector);
+
+				FComputeShaderUtils::AddPass(
+					GraphBuilder,
+					RDG_EVENT_NAME("NNE.Operator.Hlsl.CumSum.Dispatch"),
+					ERDGPassFlags::Compute | ERDGPassFlags::NeverCull,
+					ComputeShader,
+					Params,
+					ThreadGroupCount);
+
+			}
+
+		}
+	};
+
+	template<int Version>
+	bool ValidateCumSumOperator(const NNE::FAttributeMap& AttributeMap, TConstArrayView<ENNETensorDataType> InputTypes, TConstArrayView<NNE::FSymbolicTensorShape> InputShapes)
+	{
+		bool bIsValid = true;
+
+		// This matches versions 11 and 14 of the CumSum operator
+		//
+		// [REDACTED_PROFILE_URL]/onnx/blob/main/docs/Changelog.md#CumSum-11
+		// [REDACTED_PROFILE_URL]/onnx/blob/main/docs/Changelog.md#CumSum-14
+
+		FAttributeValidator AttributeValidator;
+		AttributeValidator.AddOptional(TEXT("exclusive"), ENNEAttributeDataType::Int32);
+		AttributeValidator.AddOptional(TEXT("reverse"), ENNEAttributeDataType::Int32);
+		
+		bIsValid &= AttributeValidator.Validate(AttributeMap);
+
+		FInputValidator InputValidator;
+		InputValidator.SetTemplateCount(2);
+		InputValidator.AddSupportedType(ENNETensorDataType::Float);
+		InputValidator.AddSupportedType(ENNETensorDataType::Int64, /*TemplateIdx*/ 1);
+		InputValidator.AddSupportedType(ENNETensorDataType::Int32, /*TemplateIdx*/ 1);
+		InputValidator.AddRequired();
+		InputValidator.AddRequired(/*TemplateIdx*/ 1);
+		
+		bIsValid &= InputValidator.Validate(InputTypes);
+
+		return bIsValid;
+	}
+
+	template<int Version>
+	FOperatorHlsl* CreateCumSumOperator()
+	{
+		return new FCumSumOperator<Version>();
+	}
+
+	bool RegisterCumSumOperator(FOperatorRegistryHlsl& Registry)
+	{
+		using namespace UE::NNEHlslShaders::Internal;
+
+		Registry.OpAdd({ {TEXT("CumSum"), TEXT("Onnx")}, 11}, CreateCumSumOperator<11>, ValidateCumSumOperator<11>);
+		Registry.OpAdd({ {TEXT("CumSum"), TEXT("Onnx")}, 14}, CreateCumSumOperator<14>, ValidateCumSumOperator<14>);
+		
+		return true;
+	}
+} // UE::NNERuntimeRDG::Private::Hlsl
diff --git a/Engine/Plugins/Experimental/NNERuntimeRDG/Source/NNERuntimeRDG/Private/Hlsl/NNERuntimeRDGCumSum.h b/Engine/Plugins/Experimental/NNERuntimeRDG/Source/NNERuntimeRDG/Private/Hlsl/NNERuntimeRDGCumSum.h
new file mode 100644
--- /dev/null
+++ b/Engine/Plugins/Experimental/NNERuntimeRDG/Source/NNERuntimeRDG/Private/Hlsl/NNERuntimeRDGCumSum.h
@@ -0,0 +1,10 @@
+// Copyright Epic Games, Inc. All Rights Reserved.
+
+#pragma once
+
+#include "NNERuntimeRDGHlslOp.h"
+
+namespace UE::NNERuntimeRDG::Private::Hlsl
+{
+	bool RegisterCumSumOperator(FOperatorRegistryHlsl& Registry);
+} // UE::NNERuntimeRDG::Private::Hlsl
diff --git a/Engine/Plugins/Experimental/NNERuntimeRDG/Source/NNERuntimeRDG/Private/NNERuntimeRDGHlsl.cpp b/Engine/Plugins/Experimental/NNERuntimeRDG/Source/NNERuntimeRDG/Private/NNERuntimeRDGHlsl.cpp
--- a/Engine/Plugins/Experimental/NNERuntimeRDG/Source/NNERuntimeRDG/Private/NNERuntimeRDGHlsl.cpp
+++ b/Engine/Plugins/Experimental/NNERuntimeRDG/Source/NNERuntimeRDG/Private/NNERuntimeRDGHlsl.cpp
@@ -21,6 +21,7 @@
 #include "Hlsl/NNERuntimeRDGConv.h"
 #include "Hlsl/NNERuntimeRDGConcat.h"
 #include "Hlsl/NNERuntimeRDGConvTranspose.h"
+#include "Hlsl/NNERuntimeRDGCumSum.h"
 #include "Hlsl/NNERuntimeRDGDropout.h"
 #include "Hlsl/NNERuntimeRDGElementWiseBinary.h"
 #include "Hlsl/NNERuntimeRDGElementWiseUnary.h"
@@ -65,6 +66,7 @@ bool UNNERuntimeRDGHlslImpl::Init()
 	RegisterConvOperator(*Registry);
 	RegisterConcatOperator(*Registry);
 	RegisterConvTransposeOperator(*Registry);
+	RegisterCumSumOperator(*Registry);
 	RegisterDropoutOperator(*Registry);
 	RegisterElementWiseBinaryOperators(*Registry);
 	RegisterElementWiseUnaryOperators(*Registry);
