[NNERuntimeRDG HLSL] Add CVarHlslOptimization

Problem: The ORT optimization pass creates additional cast operators in order to be compatible with the CPU execution provider. This makes it impossible to test fp16 support.

Solution: Add a CVar to disable this optimization pass during developement of fp16 support.

Tested:
* Win64 Standalone tests: all
* Win64 Editor tests: all

#rnx

diff --git a/Engine/Plugins/Experimental/NNERuntimeRDG/Source/NNERuntimeRDG/Private/NNERuntimeRDGModule.cpp b/Engine/Plugins/Experimental/NNERuntimeRDG/Source/NNERuntimeRDG/Private/NNERuntimeRDGModule.cpp
--- a/Engine/Plugins/Experimental/NNERuntimeRDG/Source/NNERuntimeRDG/Private/NNERuntimeRDGModule.cpp
+++ b/Engine/Plugins/Experimental/NNERuntimeRDG/Source/NNERuntimeRDG/Private/NNERuntimeRDGModule.cpp
@@ -1,10 +1,21 @@
 // Copyright Epic Games, Inc. All Rights Reserved.
 
+#include "HAL/IConsoleManager.h"
 #include "NNERuntimeRDGModule.h"
 #include "NNE.h"
 #include "NNERuntimeRDGHlsl.h"
 #include "UObject/WeakInterfacePtr.h"
 
+static TAutoConsoleVariable<int32> CVarHlslModelOptimization(
+	TEXT("nne.hlsl.ModelOptimization"),
+	1,
+	TEXT("Allows model optimizations when model are cooked for the HLSL runtime.\n")
+	TEXT(" 0: disabled\n")
+	TEXT(" 1: enabled (default)")
+#if !WITH_EDITOR
+	,ECVF_ReadOnly
+#endif
+	);
 
 void FNNERuntimeRDGModule::StartupModule()
 {
diff --git a/Engine/Plugins/Experimental/NNERuntimeRDG/Source/NNERuntimeRDGUtils/Private/NNERuntimeRDGUtilsModelOptimizerONNX.cpp b/Engine/Plugins/Experimental/NNERuntimeRDG/Source/NNERuntimeRDGUtils/Private/NNERuntimeRDGUtilsModelOptimizerONNX.cpp
--- a/Engine/Plugins/Experimental/NNERuntimeRDG/Source/NNERuntimeRDGUtils/Private/NNERuntimeRDGUtilsModelOptimizerONNX.cpp
+++ b/Engine/Plugins/Experimental/NNERuntimeRDG/Source/NNERuntimeRDGUtils/Private/NNERuntimeRDGUtilsModelOptimizerONNX.cpp
@@ -2,6 +2,7 @@
 
 #include "NNERuntimeRDGUtilsModelOptimizerONNX.h"
 
+#include "HAL/IConsoleManager.h"
 #include "HAL/PlatformFileManager.h"
 #include "HAL/PlatformTime.h"
 #include "Misc/FileHelper.h"
@@ -40,6 +41,12 @@ public:
 			return false;
 		}
 
+		static const auto CVarHlslModelOptimization = IConsoleManager::Get().FindTConsoleVariableDataInt(TEXT("nne.hlsl.ModelOptimization"));
+		if (CVarHlslModelOptimization && CVarHlslModelOptimization->GetValueOnAnyThread() == 0)
+		{
+			return true;
+		}
+
 		FString ProjIntermediateDir = FPaths::ConvertRelativePathToFull(FPaths::ProjectIntermediateDir());
 		FString ModelToOptimizePath = FPaths::CreateTempFilename(*ProjIntermediateDir, TEXT("ORTOptimizerPass_ToOptimize"), TEXT(".onnx"));
 		FString TargetExtension = TargetFormat == ENNEInferenceFormat::ONNX ? TEXT(".onnx") : TEXT(".ort");
