[NNERuntimeRDG HLSL] Add fp16 support for the cast operator

Problem: The ORT optimization pass creates additional cast operators in order to be compatible with the CPU execution provider. This means we need to have fp16 support in the cast operator once we have fp16 models in our tests.

Solution: Add support of static fp16 and dynamic float, fp16 and int32 support to the cast operator

Tested:
* Win64 Standalone tests: all
* Win64 Editor tests: all
* Consoles: smoke, cast operator

#rnx

#virtualized

diff --git a/Engine/Plugins/Experimental/NNERuntimeRDG/Shaders/Private/NNEHlslShaders/NNEHlslShadersCast.usf b/Engine/Plugins/Experimental/NNERuntimeRDG/Shaders/Private/NNEHlslShaders/NNEHlslShadersCast.usf
new file mode 100644
--- /dev/null
+++ b/Engine/Plugins/Experimental/NNERuntimeRDG/Shaders/Private/NNEHlslShaders/NNEHlslShadersCast.usf
@@ -0,0 +1,18 @@
+// Copyright Epic Games, Inc. All Rights Reserved.
+
+#include "/Engine/Public/Platform.ush"
+
+Buffer<INPUT_TYPE> Input;
+RWBuffer<OUTPUT_TYPE> Output;
+uint Num;
+uint ThreadCountX;
+
+[numthreads(THREADGROUP_SIZE_X, 1, 1)]
+void Cast(in const uint3 DispatchThreadID : SV_DispatchThreadID)
+{
+	const uint Index = DispatchThreadID.y * ThreadCountX + DispatchThreadID.x;
+	if (Index < Num)
+	{
+		Output[Index] = (OUTPUT_TYPE) Input[Index];
+	}
+}
diff --git a/Engine/Plugins/Experimental/NNERuntimeRDG/Source/NNEHlslShaders/Internal/NNEHlslShadersCastCS.h b/Engine/Plugins/Experimental/NNERuntimeRDG/Source/NNEHlslShaders/Internal/NNEHlslShadersCastCS.h
new file mode 100644
--- /dev/null
+++ b/Engine/Plugins/Experimental/NNERuntimeRDG/Source/NNEHlslShaders/Internal/NNEHlslShadersCastCS.h
@@ -0,0 +1,42 @@
+// Copyright Epic Games, Inc. All Rights Reserved.
+
+#pragma once
+
+#include "NNEHlslShaderBase.h"
+#include "NNEHlslShadersTypeHelper.h"
+#include "NNEOperator.h"
+#include "RenderGraphUtils.h"
+#include "ShaderParameterUtils.h"
+
+namespace UE::NNEHlslShaders::Internal
+{
+	class FCastConstants
+	{
+	public:
+		static const int32 NUM_GROUP_THREADS{ 256 };
+	};
+
+	class NNEHLSLSHADERS_API TCastCS : public FHlslShaderBase
+	{
+		DECLARE_GLOBAL_SHADER(TCastCS);
+		SHADER_USE_PARAMETER_STRUCT(TCastCS, FHlslShaderBase)
+
+		class FInputType : SHADER_PERMUTATION_ENUM_CLASS("INPUT_TYPE_ENUM", ENNEShaderDataType);
+		class FOutputType : SHADER_PERMUTATION_ENUM_CLASS("OUTPUT_TYPE_ENUM", ENNEShaderDataType);
+		using FPermutationDomain = TShaderPermutationDomain<FInputType,FOutputType>;
+
+	public:
+		// Depending on the shader permutation 'Input' and 'Output' buffers have
+		// different element types and are not only used with uint values.
+		BEGIN_SHADER_PARAMETER_STRUCT(FParameters, )
+			SHADER_PARAMETER_RDG_BUFFER_SRV(Buffer<uint>, Input)
+			SHADER_PARAMETER_RDG_BUFFER_UAV(RWBuffer<uint>, Output)
+			SHADER_PARAMETER(uint32, Num)
+			SHADER_PARAMETER(uint32, ThreadCountX)
+		END_SHADER_PARAMETER_STRUCT()
+
+		static bool ShouldCompilePermutation(const FGlobalShaderPermutationParameters& Parameters);
+
+		static void ModifyCompilationEnvironment(const FGlobalShaderPermutationParameters& InParameters, FShaderCompilerEnvironment& OutEnvironment);
+	};
+} // UE::NNEHlslShaders::Internal
diff --git a/Engine/Plugins/Experimental/NNERuntimeRDG/Source/NNEHlslShaders/Internal/NNEHlslShadersTypeHelper.h b/Engine/Plugins/Experimental/NNERuntimeRDG/Source/NNEHlslShaders/Internal/NNEHlslShadersTypeHelper.h
new file mode 100644
--- /dev/null
+++ b/Engine/Plugins/Experimental/NNERuntimeRDG/Source/NNEHlslShaders/Internal/NNEHlslShadersTypeHelper.h
@@ -0,0 +1,29 @@
+// Copyright Epic Games, Inc. All Rights Reserved.
+
+#pragma once
+
+#include "NNETypes.h"
+#include "PixelFormat.h"
+
+namespace UE::NNEHlslShaders::Internal
+{
+	enum class ENNEShaderDataType : uint8
+	{
+		BOOL,
+		FLOAT16,
+		FLOAT32,
+		INT8,
+		INT16,
+		INT32,
+		UINT8,
+		UINT16,
+		UINT32,
+
+		MAX
+	};
+
+	const TCHAR	       NNEHLSLSHADERS_API *ShaderDataTypeToName(ENNEShaderDataType ShaderType);
+	ENNEShaderDataType NNEHLSLSHADERS_API TensorToShaderDataType(ENNETensorDataType TensorType);
+	EPixelFormat	   NNEHLSLSHADERS_API TensorDataTypeToPixelFormat(ENNETensorDataType TensorType);
+
+} // UE::NNEHlslShaders::Internal
diff --git a/Engine/Plugins/Experimental/NNERuntimeRDG/Source/NNEHlslShaders/Private/NNEHlslShadersCastCS.cpp b/Engine/Plugins/Experimental/NNERuntimeRDG/Source/NNEHlslShaders/Private/NNEHlslShadersCastCS.cpp
new file mode 100644
--- /dev/null
+++ b/Engine/Plugins/Experimental/NNERuntimeRDG/Source/NNEHlslShaders/Private/NNEHlslShadersCastCS.cpp
@@ -0,0 +1,48 @@
+// Copyright Epic Games, Inc. All Rights Reserved.
+
+#include "NNEHlslShadersCastCS.h"
+
+#include "NNE.h"
+#include "NNEHlslShadersTypeHelper.h"
+#include "ShaderCompilerCore.h"
+
+namespace UE::NNEHlslShaders::Internal
+{
+	bool TCastCS::ShouldCompilePermutation(const FGlobalShaderPermutationParameters& Parameters)
+	{
+		if (!FHlslShaderBase::ShouldCompilePermutation(Parameters))
+		{
+			return false;
+		}
+
+		static TArray<ENNEShaderDataType> SupportedTypes =
+		{
+			ENNEShaderDataType::FLOAT16,
+			ENNEShaderDataType::FLOAT32,
+			ENNEShaderDataType::INT32
+		};
+
+		FPermutationDomain PermutationVector(Parameters.PermutationId);
+		ENNEShaderDataType InputType = PermutationVector.Get<TCastCS::FInputType>();
+		ENNEShaderDataType OutputType = PermutationVector.Get<TCastCS::FOutputType>();
+
+		return SupportedTypes.Contains(InputType) && SupportedTypes.Contains(OutputType);
+	}
+
+	void TCastCS::ModifyCompilationEnvironment(const FGlobalShaderPermutationParameters& InParameters, FShaderCompilerEnvironment& OutEnvironment)
+	{
+		FGlobalShader::ModifyCompilationEnvironment(InParameters, OutEnvironment);
+
+		OutEnvironment.SetDefine(TEXT("THREADGROUP_SIZE_X"), FCastConstants::NUM_GROUP_THREADS);
+
+		FPermutationDomain PermutationVector(InParameters.PermutationId);
+		ENNEShaderDataType InputType = PermutationVector.Get<TCastCS::FInputType>();
+		ENNEShaderDataType OutputType = PermutationVector.Get<TCastCS::FOutputType>();
+		OutEnvironment.SetDefine(TEXT("INPUT_TYPE"),  ShaderDataTypeToName(InputType));
+		OutEnvironment.SetDefine(TEXT("OUTPUT_TYPE"), ShaderDataTypeToName(OutputType));
+
+		OutEnvironment.CompilerFlags.Add(CFLAG_AllowRealTypes);
+	}
+
+	IMPLEMENT_GLOBAL_SHADER(TCastCS, "/NNEHlslShaders/NNEHlslShadersCast.usf", "Cast", SF_Compute);
+} // UE::NNEHlslShaders::Internal
diff --git a/Engine/Plugins/Experimental/NNERuntimeRDG/Source/NNEHlslShaders/Private/NNEHlslShadersTypeHelper.cpp b/Engine/Plugins/Experimental/NNERuntimeRDG/Source/NNEHlslShaders/Private/NNEHlslShadersTypeHelper.cpp
new file mode 100644
--- /dev/null
+++ b/Engine/Plugins/Experimental/NNERuntimeRDG/Source/NNEHlslShaders/Private/NNEHlslShadersTypeHelper.cpp
@@ -0,0 +1,63 @@
+// Copyright Epic Games, Inc. All Rights Reserved.
+
+#include "NNEHlslShadersTypeHelper.h"
+
+namespace UE::NNEHlslShaders::Internal
+{
+	const TCHAR* ShaderDataTypeToName(ENNEShaderDataType ShaderType)
+	{
+		switch (ShaderType)
+		{
+			case ENNEShaderDataType::BOOL	 : return TEXT("bool");
+			case ENNEShaderDataType::FLOAT16 : return TEXT("float16_t");
+			case ENNEShaderDataType::FLOAT32 : return TEXT("float");
+			case ENNEShaderDataType::INT8	 : return TEXT("int8_t");
+			case ENNEShaderDataType::INT16	 : return TEXT("int16_t");
+			case ENNEShaderDataType::INT32	 : return TEXT("int32_t");
+			case ENNEShaderDataType::UINT8	 : return TEXT("uint8_t");
+			case ENNEShaderDataType::UINT16	 : return TEXT("uint16_t");
+			case ENNEShaderDataType::UINT32	 : return TEXT("uint32_t");
+			default:
+				check(false);
+				return nullptr;
+		}
+	}
+
+	ENNEShaderDataType TensorToShaderDataType(ENNETensorDataType TensorType)
+	{
+		switch (TensorType)
+		{
+			case ENNETensorDataType::Boolean:    return ENNEShaderDataType::BOOL;
+			case ENNETensorDataType::Half:		 return ENNEShaderDataType::FLOAT16;
+			case ENNETensorDataType::Float:		 return ENNEShaderDataType::FLOAT32;
+			case ENNETensorDataType::Int8:		 return ENNEShaderDataType::INT8;
+			case ENNETensorDataType::Int16:		 return ENNEShaderDataType::INT16;
+			case ENNETensorDataType::Int32:		 return ENNEShaderDataType::INT32;
+			case ENNETensorDataType::UInt8:		 return ENNEShaderDataType::UINT8;
+			case ENNETensorDataType::UInt16:	 return ENNEShaderDataType::UINT16;
+			case ENNETensorDataType::UInt32:	 return ENNEShaderDataType::UINT32;
+			default:
+				check(false);
+				return (ENNEShaderDataType)0;
+		}
+	}
+
+	EPixelFormat TensorDataTypeToPixelFormat(ENNETensorDataType TensorType)
+	{
+		switch (TensorType)
+		{
+			case ENNETensorDataType::Boolean:    return EPixelFormat::PF_R8;
+			case ENNETensorDataType::Half:		 return EPixelFormat::PF_R16F;
+			case ENNETensorDataType::Float:		 return EPixelFormat::PF_R32_FLOAT;
+			case ENNETensorDataType::Int8:		 return EPixelFormat::PF_R8_SINT;
+			case ENNETensorDataType::Int16:		 return EPixelFormat::PF_R16_SINT;
+			case ENNETensorDataType::Int32:		 return EPixelFormat::PF_R32_SINT;
+			case ENNETensorDataType::UInt8:		 return EPixelFormat::PF_R8_UINT;
+			case ENNETensorDataType::UInt16:	 return EPixelFormat::PF_R16_UINT;
+			case ENNETensorDataType::UInt32:	 return EPixelFormat::PF_R32_UINT;
+			default:
+				check(false);
+				return EPixelFormat::PF_Unknown;
+		}
+	}
+} // UE::NNEHlslShaders::Internal
diff --git a/Engine/Plugins/Experimental/NNERuntimeRDG/Source/NNERuntimeRDG/Private/Helper/NNERuntimeRDGHelperCast.cpp b/Engine/Plugins/Experimental/NNERuntimeRDG/Source/NNERuntimeRDG/Private/Helper/NNERuntimeRDGHelperCast.cpp
--- a/Engine/Plugins/Experimental/NNERuntimeRDG/Source/NNERuntimeRDG/Private/Helper/NNERuntimeRDGHelperCast.cpp
+++ b/Engine/Plugins/Experimental/NNERuntimeRDG/Source/NNERuntimeRDG/Private/Helper/NNERuntimeRDGHelperCast.cpp
@@ -28,6 +28,11 @@ namespace UE::NNERuntimeRDG::Internal::CPUHelper::Cast
 				ApplyWithResolvedTypes(Tensor.GetPreparedData<float>(), OutputData);
 				break;
 			}
+			case ENNETensorDataType::Half:
+			{
+				ApplyWithResolvedTypes(Tensor.GetPreparedData<FFloat16>(), OutputData);
+				break;
+			}
 			case ENNETensorDataType::Int32:
 			{
 				ApplyWithResolvedTypes(Tensor.GetPreparedData<int32>(), OutputData);
@@ -59,6 +64,11 @@ namespace UE::NNERuntimeRDG::Internal::CPUHelper::Cast
 				ApplyWithResolvedOutputType<float>(Tensor, OutputTensor);
 				break;
 			}
+			case ENNETensorDataType::Half:
+			{
+				ApplyWithResolvedOutputType<FFloat16>(Tensor, OutputTensor);
+				break;
+			}
 			case ENNETensorDataType::Int32:
 			{
 				ApplyWithResolvedOutputType<int32>(Tensor, OutputTensor);
@@ -72,5 +82,4 @@ namespace UE::NNERuntimeRDG::Internal::CPUHelper::Cast
 			}
 		}
 	}
-	
-} // UE::NNERuntimeRDG::Internal::CPUHelper::Cast
\ No newline at end of file
+} // UE::NNERuntimeRDG::Internal::CPUHelper::Cast
diff --git a/Engine/Plugins/Experimental/NNERuntimeRDG/Source/NNERuntimeRDG/Private/Hlsl/NNERuntimeRDGCast.cpp b/Engine/Plugins/Experimental/NNERuntimeRDG/Source/NNERuntimeRDG/Private/Hlsl/NNERuntimeRDGCast.cpp
--- a/Engine/Plugins/Experimental/NNERuntimeRDG/Source/NNERuntimeRDG/Private/Hlsl/NNERuntimeRDGCast.cpp
+++ b/Engine/Plugins/Experimental/NNERuntimeRDG/Source/NNERuntimeRDG/Private/Hlsl/NNERuntimeRDGCast.cpp
@@ -3,12 +3,17 @@
 #include "NNERuntimeRDGCast.h"
 
 #include "Helper/NNERuntimeRDGLogHelper.h"
+#include "Helper/NNERuntimeRDGOperatorHelper.h"
+#include "NNEHlslShadersCastCS.h"
+#include "NNEHlslShadersTypeHelper.h"
 #include "NNERuntimeRDGHelperCast.h"
+#include "NNERuntimeRDGHlslHelper.h"
 #include "NNETensor.h"
 #include "NNETypes.h"
 
 namespace UE::NNERuntimeRDG::Private::Hlsl
 {
+	DECLARE_GPU_STAT_NAMED(FNNEOperatorCast, TEXT("NNE.Operator.Hlsl.Cast"));
 	/**
 	 * Cast operator implementation
 	 */
@@ -25,13 +30,16 @@ namespace UE::NNERuntimeRDG::Private::Hlsl
 			check(OutputTensors.Num() == 1);
 			OutputTensors[0]->SetShape(InputTensors[0]->GetShape());
 
-			const NNE::Internal::FTensor& X = *InputTensors[0];
+			const NNE::Internal::FTensor& Input = *InputTensors[0];
+			NNE::Internal::FTensor& Output = *OutputTensors[0];
 
-			Internal::CPUHelper::Cast::Apply(X, *OutputTensors[0]);
+			Internal::CPUHelper::Cast::Apply(Input, Output);
 
-			if (!OutputTensors[0]->HasPreparedData())
+			bool Has64BitDataType = Input.GetDataType() == ENNETensorDataType::Int64 || Output.GetDataType() == ENNETensorDataType::UInt64;
+
+			if (!Output.HasPreparedData() && Has64BitDataType)
 			{
-				UE_LOG(LogNNE, Warning, TEXT("Cast: Output could not be computed as a constant tensor, however Cast is not implemented on GPU at the moment."));
+				UE_LOG(LogNNE, Warning, TEXT("Cast: Output could not be computed as a constant tensor, however Cast doesn't support dynamic 64 bit tensor types."));
 				return -1;
 			}
 
@@ -57,7 +65,44 @@ namespace UE::NNERuntimeRDG::Private::Hlsl
 
 		virtual void Dispatch(FRDGBuilder& GraphBuilder, TConstArrayView<FTensorRDGRef> InputTensors, TConstArrayView<FTensorRDGRef> OutputTensors) override
 		{
-			UE_LOG(LogNNE, Warning, TEXT("Cast: Output should be constant and already uploaded to GPU memory. Dispatch should not need to be called."));
+			using namespace UE::NNEHlslShaders::Internal;
+			
+			check(InputTensors.Num() == 1);
+			check(OutputTensors.Num() == 1);
+			check(InputTensors[0] != nullptr);
+			
+			check(OutputTensors[0] != nullptr);
+			const FTensorRDG& Input = *InputTensors[0];
+			const FTensorRDG& Output = *OutputTensors[0];
+
+			FRDGBufferSRVRef InputSRV = GraphBuilder.CreateSRV(FRDGBufferSRVDesc(Input.GetBuffer(), TensorDataTypeToPixelFormat(Input.GetDataType())));
+			FRDGBufferUAVRef OutputUAV = GraphBuilder.CreateUAV(FRDGBufferUAVDesc(Output.GetBuffer(), TensorDataTypeToPixelFormat(Output.GetDataType())));
+
+			FIntVector ThreadGroupCount = ComputeElementWiseThreadGroups(Output.GetVolume(), FCastConstants::NUM_GROUP_THREADS);
+
+			// Set parameters
+			TCastCS::FParameters* Params = GraphBuilder.AllocParameters<TCastCS::FParameters>();
+			Params->Input = InputSRV;
+			Params->Output = OutputUAV;
+			Params->Num = Output.GetVolume();
+			Params->ThreadCountX = ThreadGroupCount.X * FCastConstants::NUM_GROUP_THREADS;
+
+			TCastCS::FPermutationDomain PermutationVector;
+			PermutationVector.Set<TCastCS::FInputType>(TensorToShaderDataType(Input.GetDataType()));
+			PermutationVector.Set<TCastCS::FOutputType>(TensorToShaderDataType(Output.GetDataType()));
+
+			TShaderMapRef<TCastCS> ComputeShader(GetGlobalShaderMap(GMaxRHIFeatureLevel), PermutationVector);
+
+			RDG_EVENT_SCOPE(GraphBuilder, "NNE.Operator.Hlsl.Cast");
+			RDG_GPU_STAT_SCOPE(GraphBuilder, FNNEOperatorCast);
+
+			FComputeShaderUtils::AddPass(
+				GraphBuilder,
+				RDG_EVENT_NAME("NNE.Operator.Hlsl.Cast.Dispatch"),
+				ERDGPassFlags::Compute | ERDGPassFlags::NeverCull,
+				ComputeShader,
+				Params,
+				ThreadGroupCount);
 		}
 	};
 
@@ -74,10 +119,9 @@ namespace UE::NNERuntimeRDG::Private::Hlsl
 			ENNETensorDataType To = (ENNETensorDataType)AttributeMap.GetValue<int32>(TEXT("to"));
 			switch (To)
 			{
+				case ENNETensorDataType::Half:
 				case ENNETensorDataType::Float:
-					break;
 				case ENNETensorDataType::Int32:
-					break;
 				case ENNETensorDataType::Int64:
 					break;
 				default:
@@ -88,6 +132,7 @@ namespace UE::NNERuntimeRDG::Private::Hlsl
 		}
 		
 		FInputValidator InputValidator;
+		InputValidator.AddSupportedType(ENNETensorDataType::Half);
 		InputValidator.AddSupportedType(ENNETensorDataType::Float);
 		InputValidator.AddSupportedType(ENNETensorDataType::Int32);
 		InputValidator.AddSupportedType(ENNETensorDataType::Int64);
