[NNERuntimeRDG HLSL] Add fp16 support for elementwise unary operators

Problem: Currently the elementwise unary operators only support fp32.

Solution: Change dispatch so that FP16 can be read and writen by the shader without changing the shader.

Tested:
* Win64 Standalone tests: all
* Win64 Editor tests: all
* Consoles: ElementWiseUnary

#rnx

#virtualized

diff --git a/Engine/Plugins/Experimental/NNERuntimeRDG/Source/NNEHlslShaders/Private/NNEHlslShadersElementWiseUnaryCS.cpp b/Engine/Plugins/Experimental/NNERuntimeRDG/Source/NNEHlslShaders/Private/NNEHlslShadersElementWiseUnaryCS.cpp
--- a/Engine/Plugins/Experimental/NNERuntimeRDG/Source/NNEHlslShaders/Private/NNEHlslShadersElementWiseUnaryCS.cpp
+++ b/Engine/Plugins/Experimental/NNERuntimeRDG/Source/NNEHlslShaders/Private/NNEHlslShadersElementWiseUnaryCS.cpp
@@ -1,6 +1,7 @@
 // Copyright Epic Games, Inc. All Rights Reserved.
 
 #include "NNEHlslShadersElementWiseUnaryCS.h"
+
 #include "NNE.h"
 
 namespace UE::NNEHlslShaders::Internal
@@ -86,7 +87,7 @@ namespace UE::NNEHlslShaders::Internal
 
 		FString OpFunc = OpTable[(int32) OpType];
 
-		if (OpFunc == "")
+		if (OpFunc == FString(""))
 		{
 			UE_LOG(LogNNE, Warning, TEXT("Undefined ElementWise Unary operator name for operator:%d"), int(OpType));
 		}
diff --git a/Engine/Plugins/Experimental/NNERuntimeRDG/Source/NNERuntimeRDG/Private/Helper/NNERuntimeRDGHelperElementWiseUnary.cpp b/Engine/Plugins/Experimental/NNERuntimeRDG/Source/NNERuntimeRDG/Private/Helper/NNERuntimeRDGHelperElementWiseUnary.cpp
--- a/Engine/Plugins/Experimental/NNERuntimeRDG/Source/NNERuntimeRDG/Private/Helper/NNERuntimeRDGHelperElementWiseUnary.cpp
+++ b/Engine/Plugins/Experimental/NNERuntimeRDG/Source/NNERuntimeRDG/Private/Helper/NNERuntimeRDGHelperElementWiseUnary.cpp
@@ -1,8 +1,10 @@
 // Copyright Epic Games, Inc. All Rights Reserved.
 
 #include "NNERuntimeRDGHelperElementWiseUnary.h"
-#include "NNETensor.h"
+
 #include "Math/UnrealMathUtility.h"
+#include "NNETensor.h"
+#include "NNETypes.h"
 
 namespace UE::NNERuntimeRDG::Internal::CPUHelper::ElementWiseUnary
 {
@@ -133,14 +135,33 @@ namespace UE::NNERuntimeRDG::Internal::CPUHelper::ElementWiseUnary
 
 		if (Tensor.HasPreparedData() && (Tensor.GetVolume() <= MaxItemInInputTensors))
 		{
-			TConstArrayView<float> TensorData = Tensor.GetPreparedData<float>();
-			TArray<float> OutputData;
-			OutputData.Reserve(TensorData.Num());
-			for (float elem : TensorData)
+			switch (Tensor.GetDataType())
 			{
-				OutputData.Add(Apply<OpType>(elem, Alpha, Beta, Gamma));
+				case ENNETensorDataType::Float:
+				{
+					TConstArrayView<float> TensorData = Tensor.GetPreparedData<float>();
+					TArray<float> OutputData;
+					OutputData.Reserve(TensorData.Num());
+					for (float elem : TensorData)
+					{
+						OutputData.Add(Apply<OpType>(elem, Alpha, Beta, Gamma));
+					}
+					OutputTensor.SetPreparedData<float>(OutputData);
+				} break;
+				case ENNETensorDataType::Half:
+				{
+					TConstArrayView<FFloat16> TensorData = Tensor.GetPreparedData<FFloat16>();
+					TArray<FFloat16> OutputData;
+					OutputData.Reserve(TensorData.Num());
+					for (FFloat16 elem : TensorData)
+					{
+						OutputData.Add(Apply<OpType>(elem, Alpha, Beta, Gamma));
+					}
+					OutputTensor.SetPreparedData<FFloat16>(OutputData);
+				} break;
+				default:
+					break;
 			}
-			OutputTensor.SetPreparedData<float>(OutputData);
 		}
 	}
 
diff --git a/Engine/Plugins/Experimental/NNERuntimeRDG/Source/NNERuntimeRDG/Private/Hlsl/NNERuntimeRDGElementWiseUnary.cpp b/Engine/Plugins/Experimental/NNERuntimeRDG/Source/NNERuntimeRDG/Private/Hlsl/NNERuntimeRDGElementWiseUnary.cpp
--- a/Engine/Plugins/Experimental/NNERuntimeRDG/Source/NNERuntimeRDG/Private/Hlsl/NNERuntimeRDGElementWiseUnary.cpp
+++ b/Engine/Plugins/Experimental/NNERuntimeRDG/Source/NNERuntimeRDG/Private/Hlsl/NNERuntimeRDGElementWiseUnary.cpp
@@ -1,12 +1,14 @@
 // Copyright Epic Games, Inc. All Rights Reserved.
 
 #include "NNERuntimeRDGElementWiseUnary.h"
-#include "NNERuntimeRDGHelperElementWiseUnary.h"
+
+#include "NNEAttributeMap.h"
 #include "NNEHlslShadersElementWiseUnaryCS.h"
+#include "NNEHlslShadersTypeHelper.h"
+#include "NNERuntimeRDGHelperElementWiseUnary.h"
 #include "NNERuntimeRDGHlslHelper.h"
-#include "NNEAttributeMap.h"
-#include "NNETypes.h"
 #include "NNETensor.h"
+#include "NNETypes.h"
 
 namespace UE::NNERuntimeRDG::Private::Hlsl
 {
@@ -81,14 +83,16 @@ namespace UE::NNERuntimeRDG::Private::Hlsl
 			return true;
 		}
 
-		virtual void Dispatch(FRDGBuilder& GraphBuilder, TConstArrayView<FTensorRDGRef> InInputTensors, TConstArrayView<FTensorRDGRef> InOutputTensors) override
+		virtual void Dispatch(FRDGBuilder& GraphBuilder, TConstArrayView<FTensorRDGRef> InputTensors, TConstArrayView<FTensorRDGRef> OutputTensors) override
 		{
-			check(InInputTensors[0] != nullptr);
-			check(InOutputTensors[0] != nullptr);
-			FRDGBufferSRVRef InputSRV = GraphBuilder.CreateSRV(FRDGBufferSRVDesc(InInputTensors[0]->GetBuffer(), PF_R32_FLOAT));
-			FRDGBufferUAVRef OutputUAV = GraphBuilder.CreateUAV(FRDGBufferUAVDesc(InOutputTensors[0]->GetBuffer(), PF_R32_FLOAT));
+			using namespace UE::NNEHlslShaders::Internal;
+
+			check(InputTensors[0] != nullptr);
+			check(OutputTensors[0] != nullptr);
+			FRDGBufferSRVRef InputSRV = GraphBuilder.CreateSRV(FRDGBufferSRVDesc(InputTensors[0]->GetBuffer(), TensorDataTypeToPixelFormat(InputTensors[0]->GetDataType())));
+			FRDGBufferUAVRef OutputUAV = GraphBuilder.CreateUAV(FRDGBufferUAVDesc(OutputTensors[0]->GetBuffer(), TensorDataTypeToPixelFormat(OutputTensors[0]->GetDataType())));
 		
-			int32 NumElements = InOutputTensors[0]->GetVolume();
+			int32 NumElements = OutputTensors[0]->GetVolume();
 			FIntVector ThreadGroupCount = ComputeElementWiseThreadGroups(NumElements, FElementWiseUnaryConstants::NUM_GROUP_THREADS);
 
 			// Set parameters
@@ -107,31 +111,30 @@ namespace UE::NNERuntimeRDG::Private::Hlsl
 			PermutationVector.Set<TElementWiseUnaryCS::FAlphaOnGPU>(false);
 			PermutationVector.Set<TElementWiseUnaryCS::FBetaOnGPU>(false);
 
-			if(OpType == NNE::Internal::EElementWiseUnaryOperatorType::Clip && Version >= 11 && InInputTensors.Num() >= 2)
+			if(OpType == NNE::Internal::EElementWiseUnaryOperatorType::Clip && Version >= 11 && InputTensors.Num() >= 2)
 			{
 				{
-					const NNE::Internal::FTensor& MinTensor = *InInputTensors[1];
+					const NNE::Internal::FTensor& MinTensor = *InputTensors[1];
 					if(!MinTensor.HasPreparedData())
 					{
 						PermutationVector.Set<TElementWiseUnaryCS::FAlphaOnGPU>(true);
-						FRDGBufferSRVRef AlphaSRV = GraphBuilder.CreateSRV(FRDGBufferSRVDesc(InInputTensors[1]->GetBuffer(), PF_R32_FLOAT));
+						FRDGBufferSRVRef AlphaSRV = GraphBuilder.CreateSRV(FRDGBufferSRVDesc(InputTensors[1]->GetBuffer(), PF_R32_FLOAT));
 						Params->AlphaTensor = AlphaSRV;
 					}
 				}
 
-				if(InInputTensors.Num() == 3)
+				if(InputTensors.Num() == 3)
 				{
-					const NNE::Internal::FTensor& MaxTensor = *InInputTensors[2];
+					const NNE::Internal::FTensor& MaxTensor = *InputTensors[2];
 					if(!MaxTensor.HasPreparedData())
 					{
 						PermutationVector.Set<TElementWiseUnaryCS::FBetaOnGPU>(true);
-						FRDGBufferSRVRef BetaSRV = GraphBuilder.CreateSRV(FRDGBufferSRVDesc(InInputTensors[2]->GetBuffer(), PF_R32_FLOAT));
+						FRDGBufferSRVRef BetaSRV = GraphBuilder.CreateSRV(FRDGBufferSRVDesc(InputTensors[2]->GetBuffer(), PF_R32_FLOAT));
 						Params->BetaTensor = BetaSRV;
 					}
 				}
 			}
 
-
 			TShaderMapRef<TElementWiseUnaryCS> ComputeShader(GetGlobalShaderMap(GMaxRHIFeatureLevel), PermutationVector);
 
 			RDG_EVENT_SCOPE(GraphBuilder, "NNE.Operator.Hlsl.ElementWise.Unary");
@@ -193,6 +196,7 @@ namespace UE::NNERuntimeRDG::Private::Hlsl
 
 		FInputValidator InputValidator;
 		InputValidator.AddSupportedType(ENNETensorDataType::Float);
+		InputValidator.AddSupportedType(ENNETensorDataType::Half);
 		InputValidator.AddRequired();
 		bIsValid &= InputValidator.Validate(InputTypes);
 
@@ -211,6 +215,7 @@ namespace UE::NNERuntimeRDG::Private::Hlsl
 
 		FInputValidator InputValidator;
 		InputValidator.AddSupportedType(ENNETensorDataType::Float);
+		InputValidator.AddSupportedType(ENNETensorDataType::Half);
 		InputValidator.AddRequired();
 		bIsValid &= InputValidator.Validate(InputTypes);
 
@@ -228,6 +233,7 @@ namespace UE::NNERuntimeRDG::Private::Hlsl
 
 		FInputValidator InputValidator;
 		InputValidator.AddSupportedType(ENNETensorDataType::Float);
+		InputValidator.AddSupportedType(ENNETensorDataType::Half);
 		InputValidator.AddRequired();
 		bIsValid &= InputValidator.Validate(InputTypes);
 
@@ -246,6 +252,7 @@ namespace UE::NNERuntimeRDG::Private::Hlsl
 
 		FInputValidator InputValidator;
 		InputValidator.AddSupportedType(ENNETensorDataType::Float);
+		InputValidator.AddSupportedType(ENNETensorDataType::Half);
 		InputValidator.AddRequired();
 		bIsValid &= InputValidator.Validate(InputTypes);
 
@@ -263,6 +270,7 @@ namespace UE::NNERuntimeRDG::Private::Hlsl
 
 		FInputValidator InputValidator;
 		InputValidator.AddSupportedType(ENNETensorDataType::Float);
+		InputValidator.AddSupportedType(ENNETensorDataType::Half);
 		InputValidator.AddRequired();
 		bIsValid &= InputValidator.Validate(InputTypes);
 
@@ -284,6 +292,7 @@ namespace UE::NNERuntimeRDG::Private::Hlsl
 
 		FInputValidator InputValidator;
 		InputValidator.AddSupportedType(ENNETensorDataType::Float);
+		InputValidator.AddSupportedType(ENNETensorDataType::Half);
 		InputValidator.AddRequired();
 		if(Version >= 11)
 		{
