[NNERuntimeRDG HLSL] Add fp16 support for OIDN operators

Problem: Currently only fp32 OIDN models are supported

Solution: Add support for fp16 to the following operators: Conv, MaxPool, Concat and Resize without changing there shaders.

Tested:
* Win64 Standalone tests: all
* Win64 Editor tests: all
* Consoles: Other Tests

#rnx

#virtualized

diff --git a/Engine/Plugins/Experimental/NNERuntimeRDG/Source/NNERuntimeRDG/Private/Hlsl/NNERuntimeRDGConcat.cpp b/Engine/Plugins/Experimental/NNERuntimeRDG/Source/NNERuntimeRDG/Private/Hlsl/NNERuntimeRDGConcat.cpp
--- a/Engine/Plugins/Experimental/NNERuntimeRDG/Source/NNERuntimeRDG/Private/Hlsl/NNERuntimeRDGConcat.cpp
+++ b/Engine/Plugins/Experimental/NNERuntimeRDG/Source/NNERuntimeRDG/Private/Hlsl/NNERuntimeRDGConcat.cpp
@@ -109,9 +109,6 @@ namespace UE::NNERuntimeRDG::Private::Hlsl
 			return true;
 		}
 
-		
-
-
 		virtual void Dispatch(FRDGBuilder& GraphBuilder, TConstArrayView<FTensorRDGRef> InputTensors, TConstArrayView<FTensorRDGRef> OutputTensors) override
 		{
 			check(InputTensors.Num() >= 1);
@@ -162,9 +159,14 @@ namespace UE::NNERuntimeRDG::Private::Hlsl
 		}
 		for (int32 i = 0; i < InputTypes.Num(); ++i)
 		{
-			if (InputTypes[i] != ENNETensorDataType::Float)
+			if (InputTypes[i] != ENNETensorDataType::Float && InputTypes[i] != ENNETensorDataType::Half)
+			{
+				UE_LOG(LogNNE, Warning, TEXT("Concat operator input '%d' of type '%d' is not supported, should be float or half at the moment."), i, int(InputTypes[i]));
+				bIsValid = false;
+			}
+			if (InputTypes[i] != InputTypes[0])
 			{
-				UE_LOG(LogNNE, Warning, TEXT("Concat operator input '%d' of type '%d' is not supported, should be float at the moment."), i, int(InputTypes[i]));
+				UE_LOG(LogNNE, Warning, TEXT("Concat operator input '%d' of type '%d' does not match type '%d' of the first input."), i, int(InputTypes[i]), int(InputTypes[0]));
 				bIsValid = false;
 			}
 		}
diff --git a/Engine/Plugins/Experimental/NNERuntimeRDG/Source/NNERuntimeRDG/Private/Hlsl/NNERuntimeRDGConv.cpp b/Engine/Plugins/Experimental/NNERuntimeRDG/Source/NNERuntimeRDG/Private/Hlsl/NNERuntimeRDGConv.cpp
--- a/Engine/Plugins/Experimental/NNERuntimeRDG/Source/NNERuntimeRDG/Private/Hlsl/NNERuntimeRDGConv.cpp
+++ b/Engine/Plugins/Experimental/NNERuntimeRDG/Source/NNERuntimeRDG/Private/Hlsl/NNERuntimeRDGConv.cpp
@@ -7,6 +7,7 @@
 #include "NNETypes.h"
 #include "NNEHlslShadersConvCS.h"
 #include "NNEHlslShadersConvMatmulCS.h"
+#include "NNEHlslShadersTypeHelper.h"
 #include "NNERuntimeRDGHelperTranspose.h"
 #include "NNERuntimeRDGHlslHelper.h"
 
@@ -43,6 +44,7 @@ namespace UE::NNERuntimeRDG::Private::Hlsl
 		TArray<int32> Pads;
 		TArray<int32> Strides;
 		bool bAreWeightsTransposed = false;
+		EPixelFormat BufferPixelFormat;
 
 	public:
 
@@ -109,6 +111,7 @@ namespace UE::NNERuntimeRDG::Private::Hlsl
 				Pads = Attributes.GetValueOrDefault<TArray<int32>>(TEXT("pads"), PadsDefault);
 			}
 			Strides = Attributes.GetValueOrDefault<TArray<int32>>(TEXT("strides"), DilationsOrStridesDefault);
+			BufferPixelFormat = TensorDataTypeToPixelFormat(Input.GetDataType());
 
 			return true;
 		}
@@ -125,12 +128,12 @@ namespace UE::NNERuntimeRDG::Private::Hlsl
 			// Set parameters
 			FConvCS::FParameters* Params = GraphBuilder.AllocParameters<FConvCS::FParameters>();
 			FConvCS::FillInParameters(GroupSize, Input.GetShape().GetData(), Weights.GetShape().GetData(), HasBias, AutoPad, Group, Dilations,Strides, Pads, *Params);
-			Params->X = GraphBuilder.CreateSRV(FRDGBufferSRVDesc(Input.GetBuffer(), PF_R32_FLOAT));
-			Params->W = GraphBuilder.CreateSRV(FRDGBufferSRVDesc(Weights.GetBuffer(), PF_R32_FLOAT));
+			Params->X = GraphBuilder.CreateSRV(FRDGBufferSRVDesc(Input.GetBuffer(), BufferPixelFormat));
+			Params->W = GraphBuilder.CreateSRV(FRDGBufferSRVDesc(Weights.GetBuffer(), BufferPixelFormat));
 			if (HasBias) {
-				Params->B = GraphBuilder.CreateSRV(FRDGBufferSRVDesc(Bias->GetBuffer(), PF_R32_FLOAT));
+				Params->B = GraphBuilder.CreateSRV(FRDGBufferSRVDesc(Bias->GetBuffer(), BufferPixelFormat));
 			}
-			Params->Y = GraphBuilder.CreateUAV(FRDGBufferUAVDesc(Output.GetBuffer(), PF_R32_FLOAT));
+			Params->Y = GraphBuilder.CreateUAV(FRDGBufferUAVDesc(Output.GetBuffer(), BufferPixelFormat));
 
 			FConvCS::FPermutationDomain PermutationVector;
 
@@ -199,12 +202,12 @@ namespace UE::NNERuntimeRDG::Private::Hlsl
 
 			// Set parameters
 			FConvMatmulCS::FParameters* Params = GraphBuilder.AllocParameters<FConvMatmulCS::FParameters>();
-			Params->Input = GraphBuilder.CreateSRV(FRDGBufferSRVDesc(Input.GetBuffer(), PF_R32_FLOAT));
-			Params->Weight= GraphBuilder.CreateSRV(FRDGBufferSRVDesc(Weights.GetBuffer(), PF_R32_FLOAT));
+			Params->Input = GraphBuilder.CreateSRV(FRDGBufferSRVDesc(Input.GetBuffer(), BufferPixelFormat));
+			Params->Weight= GraphBuilder.CreateSRV(FRDGBufferSRVDesc(Weights.GetBuffer(), BufferPixelFormat));
 			if (bHasBias) {
-				Params->Bias = GraphBuilder.CreateSRV(FRDGBufferSRVDesc(Bias->GetBuffer(), PF_R32_FLOAT));
+				Params->Bias = GraphBuilder.CreateSRV(FRDGBufferSRVDesc(Bias->GetBuffer(), BufferPixelFormat));
 			}
-			Params->Output = GraphBuilder.CreateUAV(FRDGBufferUAVDesc(Output.GetBuffer(), PF_R32_FLOAT));
+			Params->Output = GraphBuilder.CreateUAV(FRDGBufferUAVDesc(Output.GetBuffer(), BufferPixelFormat));
 			Params->Ci = Ci;
 			Params->Hi = Hi;
 			Params->Wi = Wi;
@@ -314,11 +317,20 @@ namespace UE::NNERuntimeRDG::Private::Hlsl
 
 		FInputValidator InputValidator;
 		InputValidator.AddSupportedType(ENNETensorDataType::Float);
+		InputValidator.AddSupportedType(ENNETensorDataType::Half);
 		InputValidator.AddRequired();
 		InputValidator.AddRequired();
 		InputValidator.AddOptional();
 		bIsValid &= InputValidator.Validate(InputTypes);
 
+		for (ENNETensorDataType InputType : InputTypes)
+		{
+			if (InputType != ENNETensorDataType::None && InputType != InputTypes[0])
+			{
+				UE_LOG(LogNNE, Warning, TEXT("All input tensor data types have to match each other"));
+				return false;
+			}
+		}
 		return bIsValid;
 	}
 
diff --git a/Engine/Plugins/Experimental/NNERuntimeRDG/Source/NNERuntimeRDG/Private/Hlsl/NNERuntimeRDGPool.cpp b/Engine/Plugins/Experimental/NNERuntimeRDG/Source/NNERuntimeRDG/Private/Hlsl/NNERuntimeRDGPool.cpp
--- a/Engine/Plugins/Experimental/NNERuntimeRDG/Source/NNERuntimeRDG/Private/Hlsl/NNERuntimeRDGPool.cpp
+++ b/Engine/Plugins/Experimental/NNERuntimeRDG/Source/NNERuntimeRDG/Private/Hlsl/NNERuntimeRDGPool.cpp
@@ -3,6 +3,7 @@
 #include "NNERuntimeRDGPool.h"
 #include "NNEHlslShadersConvCS.h"
 #include "NNEHlslShadersPoolCS.h"
+#include "NNEHlslShadersTypeHelper.h"
 #include "NNERuntimeRDGHlslHelper.h"
 #include "NNETensor.h"
 #include "NNETypes.h"
@@ -37,6 +38,7 @@ namespace UE::NNERuntimeRDG::Private::Hlsl
 		TArray<int32> KernelShape;
 		int32 CeilMode = 0; // 0 is floor, 1 is ceil
 		int32 KernelVolume = 0;
+		EPixelFormat BufferPixelFormat;
 
 	public:
 
@@ -169,6 +171,7 @@ namespace UE::NNERuntimeRDG::Private::Hlsl
 					}
 				}
 			}
+			BufferPixelFormat = UE::NNEHlslShaders::Internal::TensorDataTypeToPixelFormat(Input.GetDataType());
 
 			return true;
 		}
@@ -184,8 +187,8 @@ namespace UE::NNERuntimeRDG::Private::Hlsl
 
 			const FTensorRDG& Input = *InputTensors[0];
 			const FTensorRDG& Output = *OutputTensors[0];
-			const FRDGBufferSRVRef InputSRV = GraphBuilder.CreateSRV(FRDGBufferSRVDesc(Input.GetBuffer(), PF_R32_FLOAT));
-			const FRDGBufferUAVRef OutputUAV = GraphBuilder.CreateUAV(FRDGBufferUAVDesc(Output.GetBuffer(), PF_R32_FLOAT));
+			const FRDGBufferSRVRef InputSRV = GraphBuilder.CreateSRV(FRDGBufferSRVDesc(Input.GetBuffer(), BufferPixelFormat));
+			const FRDGBufferUAVRef OutputUAV = GraphBuilder.CreateUAV(FRDGBufferUAVDesc(Output.GetBuffer(), BufferPixelFormat));
 			const FIntVector ThreadGroupCount = ComputeElementWiseThreadGroups(Output.GetVolume(), FPoolConstants::NUM_GROUP_THREADS);
 
 			// Set parameters
@@ -287,6 +290,7 @@ namespace UE::NNERuntimeRDG::Private::Hlsl
 
 		FInputValidator InputValidator;
 		InputValidator.AddSupportedType(ENNETensorDataType::Float);
+		InputValidator.AddSupportedType(ENNETensorDataType::Half);
 		InputValidator.AddRequired();
 		bIsValid &= InputValidator.Validate(InputTypes);
 
diff --git a/Engine/Plugins/Experimental/NNERuntimeRDG/Source/NNERuntimeRDG/Private/Hlsl/NNERuntimeRDGResize.cpp b/Engine/Plugins/Experimental/NNERuntimeRDG/Source/NNERuntimeRDG/Private/Hlsl/NNERuntimeRDGResize.cpp
--- a/Engine/Plugins/Experimental/NNERuntimeRDG/Source/NNERuntimeRDG/Private/Hlsl/NNERuntimeRDGResize.cpp
+++ b/Engine/Plugins/Experimental/NNERuntimeRDG/Source/NNERuntimeRDG/Private/Hlsl/NNERuntimeRDGResize.cpp
@@ -3,6 +3,7 @@
 #include "NNERuntimeRDGResize.h"
 #include "NNEHlslShadersResizeCS.h"
 
+#include "NNEHlslShadersTypeHelper.h"
 #include "NNERuntimeRDGHlslHelper.h"
 #include "NNETensor.h"
 #include "NNETypes.h"
@@ -55,64 +56,21 @@ namespace UE::NNERuntimeRDG::Private::Hlsl
 		TArray<float> ScalesData;
 		TArray<float> Adjustments; // Necessary for "half_pixel_symmetric" coordinate transform mode
 		TArray<float> RegionOfInterest;
+		EPixelFormat BufferPixelFormat;
 		
 
 	public:
 
 		virtual int PrepareOutputs(TConstArrayView<NNE::Internal::FTensorRef> InputTensors, TArrayView<NNE::Internal::FTensorRef> OutputTensors) override
 		{
-			check(InputTensors.Num() >= 1 && InputTensors.Num() <= 4);
+			check(InputTensors.Num() >= 3 && InputTensors.Num() <= 4);
 			check(OutputTensors.Num() == 1);
 
-			
-			if(InputTensors.Num() < 3)
-			{
-				UE_LOG(LogNNE, Warning, TEXT("Hlsl Resize: one of tensors `scales` or `sizes` MUST be provided."));
-				return -1;
-			}
-
-			if(InputTensors.Num() == 4 && InputTensors[2]->GetDataType() != ENNETensorDataType::None)
-			{
-				UE_LOG(LogNNE, Warning, TEXT("Hlsl Resize: both `sizes` and `scales` were provided."));
-				return -1;
-			}
-
 			const NNE::Internal::FTensor& Input = *InputTensors[0];
 			const NNE::Internal::FTensor& Roi = *InputTensors[1];
 
-			if(Input.GetShape().Rank() < 1)
-			{
-				UE_LOG(LogNNE, Warning, TEXT("Hlsl Resize: input tensor must have rank >= 1 (name %s)."), *Input.GetName());
-				return -1;
-			}
-
-			if(Input.GetDataType() != ENNETensorDataType::Float)
-			{
-				UE_LOG(LogNNE, Warning, TEXT("Hlsl Resize: only float data type supported for input tensor (name %s)."), *Input.GetName());
-				return -1;
-			}
-
-			if(Roi.GetShape().Rank() != 1)
-			{
-				UE_LOG(LogNNE, Warning, TEXT("Hlsl Resize: `roi` tensor must have rank 1 (name %s)."), *Roi.GetName());
-				return -1;
-			}
-
-			if(CoordTransMode == UE::NNEHlslShaders::Internal::ECoordTransMode::TfCropAndResize
-				&& Roi.GetShape().Volume() != 2 * Input.GetShape().Rank())
-			{
-				UE_LOG(LogNNE, Warning, TEXT("Hlsl Resize: `roi` tensor (name %s) must have 2 * N length."), *Roi.GetName());
-				return -1;
-			}
-
-			if(Roi.GetShape().Volume() != 0)
+			if(CoordTransMode == UE::NNEHlslShaders::Internal::ECoordTransMode::TfCropAndResize)
 			{
-				if(Roi.GetDataType() != ENNETensorDataType::Float)
-				{
-					UE_LOG(LogNNE, Warning, TEXT("Hlsl Resize: only float data type supported for`roi` tensor (name %s)."), *Roi.GetName());
-					return -1;
-				}
-
 				if(!Roi.HasPreparedData())
 				{
 					UE_LOG(LogNNE, Warning, TEXT("Hlsl Resize: `roi` tensor could not be made constant. (name %s)."), *Roi.GetName());
@@ -121,8 +79,7 @@ namespace UE::NNERuntimeRDG::Private::Hlsl
 				
 				RegionOfInterest = Roi.GetPreparedData<float>();
 
-				if(CoordTransMode == UE::NNEHlslShaders::Internal::ECoordTransMode::TfCropAndResize
-					&& RegionOfInterest.Num() != 2 * Input.GetShape().Rank())
+				if(RegionOfInterest.Num() != 2 * Input.GetShape().Rank())
 				{
 					UE_LOG(LogNNE, Warning, TEXT("Hlsl Resize: `roi` tensor (name %s) must have 2 * N length."), *Roi.GetName());
 					return -1;
@@ -137,19 +94,12 @@ namespace UE::NNERuntimeRDG::Private::Hlsl
 			if(InputTensors.Num() == 3)
 			{
 				const NNE::Internal::FTensor& Scales = *InputTensors[2];
-				if(Scales.GetDataType() != ENNETensorDataType::Float)
-				{
-					UE_LOG(LogNNE, Warning, TEXT("Hlsl Resize: `scales` tensor (name %s) must have float data type."), *Scales.GetName());
-					return -1;
-				}
 				if(!Scales.HasPreparedData())
 				{
 					UE_LOG(LogNNE, Warning, TEXT("Hlsl Resize: `scales` tensor could not be made constant. (name %s)."), *Scales.GetName());
 					return -1;
 				}
 
-				
-				
 				if(CoordTransMode == UE::NNEHlslShaders::Internal::ECoordTransMode::HalfPixelSymmetric)
 				{
 					Adjustments.SetNumUninitialized(Input.GetShape().Rank());
@@ -171,17 +121,11 @@ namespace UE::NNERuntimeRDG::Private::Hlsl
 						Adjustments[Idx] = (float) OutputShapeData[Idx] / ((float) Input.GetShape().GetData()[Idx] * Scales.GetPreparedData<float>()[Idx]);
 					}
 				}
-
 			}
 			
 			if(InputTensors.Num() == 4)
 			{
 				const NNE::Internal::FTensor& Sizes = *InputTensors[3];
-				if(Sizes.GetDataType() != ENNETensorDataType::Int64)
-				{
-					UE_LOG(LogNNE, Warning, TEXT("Hlsl Resize: `sizes` tensor (name %s) must have int64 data type."), *Sizes.GetName());
-					return -1;
-				}
 				if(!Sizes.HasPreparedData())
 				{
 					UE_LOG(LogNNE, Warning, TEXT("Hlsl Resize: `sizes` tensor could not be made constant. (name %s)."), *Sizes.GetName());
@@ -222,7 +166,6 @@ namespace UE::NNERuntimeRDG::Private::Hlsl
 			NNE::FTensorShape OutputShape = NNE::FTensorShape::Make(OutputShapeData);
 			OutputTensors[0]->SetShape(OutputShape);
 
-			
 			return 0;
 		};
 
@@ -282,6 +225,7 @@ namespace UE::NNERuntimeRDG::Private::Hlsl
 			}
 
 			NearestMode = UE::NNEHlslShaders::Internal::FResizeCS::NearestModeFromString(*Attributes.GetValueOrDefault<FString>(TEXT("nearest_mode"), TEXT("round_prefer_floor")));
+			BufferPixelFormat = UE::NNEHlslShaders::Internal::TensorDataTypeToPixelFormat(InputTensorDescs[0].GetDataType());
 			
 			return true;
 		}
@@ -304,8 +248,8 @@ namespace UE::NNERuntimeRDG::Private::Hlsl
 			RDG_EVENT_SCOPE(GraphBuilder, "NNE.Operator.Hlsl.Resize");
 			RDG_GPU_STAT_SCOPE(GraphBuilder, FNNEOperatorResize);
 
-			const FRDGBufferSRVRef InputSRV = GraphBuilder.CreateSRV(FRDGBufferSRVDesc(Input.GetBuffer(), PF_R32_FLOAT));
-			const FRDGBufferUAVRef OutputUAV = GraphBuilder.CreateUAV(FRDGBufferUAVDesc(Output.GetBuffer(), PF_R32_FLOAT));
+			const FRDGBufferSRVRef InputSRV = GraphBuilder.CreateSRV(FRDGBufferSRVDesc(Input.GetBuffer(), BufferPixelFormat));
+			const FRDGBufferUAVRef OutputUAV = GraphBuilder.CreateUAV(FRDGBufferUAVDesc(Output.GetBuffer(), BufferPixelFormat));
 
 			const FIntVector ThreadGroupCount = ComputeElementWiseThreadGroups(Output.GetVolume(), FResizeConstants::NUM_GROUP_THREADS);
 
@@ -381,40 +325,53 @@ namespace UE::NNERuntimeRDG::Private::Hlsl
 		AttributeValidator.AddOptional(TEXT("nearest_mode"), ENNEAttributeDataType::String);
 		bIsValid &= AttributeValidator.Validate(AttributeMap);
 
-		if(InputTypes.Num() < 3 || InputTypes.Num() > 4)
-		{
-			UE_LOG(LogNNE, Warning, TEXT("Got a total of '%d' inputs but should be between 3 and 4."), InputTypes.Num());
-			return false;
-		}
-		if(InputTypes[0] != ENNETensorDataType::Float)
-		{
-			UE_LOG(LogNNE, Warning, TEXT("Input tensor must be of type float."));
-			return false;
-		}
-		if(InputTypes[1] != ENNETensorDataType::Float)
+		FInputValidator InputValidator;
+		InputValidator.SetTemplateCount(3);
+		InputValidator.AddSupportedType(ENNETensorDataType::Float);
+		InputValidator.AddSupportedType(ENNETensorDataType::Half);
+		InputValidator.AddSupportedType(ENNETensorDataType::Float, 1/*TemplateIdx*/);
+		InputValidator.AddSupportedType(ENNETensorDataType::Int64, 2/*TemplateIdx*/);
+
+		InputValidator.AddRequired();
+		InputValidator.AddOptional(1/*TemplateIdx*/);
+		InputValidator.AddOptional(1/*TemplateIdx*/);
+		InputValidator.AddOptional(2/*TemplateIdx*/);
+
+		if (!InputValidator.Validate(InputTypes))
 		{
-			UE_LOG(LogNNE, Warning, TEXT("Roi tensor must be of type float."));
 			return false;
 		}
-		if(InputShapes[1].Rank() != 1)
+
+		if(InputTypes.Num() < 3 || InputTypes.Num() > 4)
 		{
-			UE_LOG(LogNNE, Warning, TEXT("Roi tensor must be a 1-D tensor."));
+			UE_LOG(LogNNE, Warning, TEXT("Got a total of '%d' inputs but should be between 3 and 4."), InputTypes.Num());
 			return false;
 		}
-		if(FResizeCS::CoordTransModeFromString(*AttributeMap.GetValueOrDefault<FString>(TEXT("coordinate_transformation_mode"), TEXT("half_pixel"))) == ECoordTransMode::TfCropAndResize &&
-			InputShapes[1].GetData()[0] != 2 * InputShapes[0].Rank())
+
+		if(InputShapes[0].Rank() < 1)
 		{
-			UE_LOG(LogNNE, Warning, TEXT("Roi tensor must have dimension 2*N (where N is the input rank) when `coordinate_transformation_mode` is `tf_crop_and_resize`."));
+			UE_LOG(LogNNE, Warning, TEXT("Hlsl Resize: input tensor must have rank >= 1."));
 			return false;
 		}
 
-		if(InputTypes.Num() == 3)
+		if (InputTypes[1] != ENNETensorDataType::None)
 		{
-			if(InputTypes[2] != ENNETensorDataType::Float)
+			if(InputShapes[1].Rank() != 1)
+			{
+				UE_LOG(LogNNE, Warning, TEXT("Roi tensor must be a 1-D tensor."));
+				return false;
+			}
+			FString TransMode = AttributeMap.GetValueOrDefault<FString>(TEXT("coordinate_transformation_mode"), TEXT("half_pixel"));
+			bool IsCropAndResizeMode = FResizeCS::CoordTransModeFromString(*TransMode) == ECoordTransMode::TfCropAndResize;
+			if(IsCropAndResizeMode && InputShapes[1].GetData()[0] != 2 * InputShapes[0].Rank())
 			{
-				UE_LOG(LogNNE, Warning, TEXT("Scales tensor must be of type float."));
+				UE_LOG(LogNNE, Warning, TEXT("Roi tensor must have dimension 2*N (where N is the input rank) when `coordinate_transformation_mode` is `tf_crop_and_resize`."));
 				return false;
 			}
+		}
+
+		if(InputTypes.Num() == 3)
+		{
 			if(InputShapes[2].Rank() != 1)
 			{
 				UE_LOG(LogNNE, Warning, TEXT("Scales tensor must be a 1-D tensor."));
@@ -433,11 +390,6 @@ namespace UE::NNERuntimeRDG::Private::Hlsl
 				UE_LOG(LogNNE, Warning, TEXT("Scales tensor must be empty (i.e. empty name and data type 'None') when Sizes is specified."));
 				return false;
 			}
-			if(InputTypes[3] != ENNETensorDataType::Int64)
-			{
-				UE_LOG(LogNNE, Warning, TEXT("Sizes tensor must be of type Int64."));
-				return false;
-			}
 			if(InputShapes[3].Rank() != 1)
 			{
 				UE_LOG(LogNNE, Warning, TEXT("Sizes tensor must be a 1-D tensor."));
