[NNERuntimeRDG HLSL] Enforce buffer sizes to be a multiple of 4 bytes

Problem: Optimizing performance for fp16 requires us to load two fp16 values at a time. This would require expensive branches or orther checks for the cases where a buffer only ha an odd number of elements.

Solution: Force the user to only provide use with buffers that are a multiple of 4 bytes and check that this is the case. Create intermediate buffers that follow the same restriction, as well as the input/output buffers in our testing framework.

Tested:
* Win64 Standalone tests: all
* Win64 Editor tests: all
* Consoles tests: Smoke

#rnx

diff --git a/Engine/Plugins/Experimental/NNERuntimeRDG/Source/NNERuntimeRDG/NNERuntimeRDG.Build.cs b/Engine/Plugins/Experimental/NNERuntimeRDG/Source/NNERuntimeRDG/NNERuntimeRDG.Build.cs
--- a/Engine/Plugins/Experimental/NNERuntimeRDG/Source/NNERuntimeRDG/NNERuntimeRDG.Build.cs
+++ b/Engine/Plugins/Experimental/NNERuntimeRDG/Source/NNERuntimeRDG/NNERuntimeRDG.Build.cs
@@ -45,5 +45,7 @@ public class NNERuntimeRDG : ModuleRules
 			PrivateDefinitions.Add("NNE_UTILITIES_AVAILABLE");
 			PrivateDependencyModuleNames.Add("NNERuntimeRDGUtils");
 		}
+
+		PublicDefinitions.Add("NNERUNTIMERDGHLSL_BUFFER_LENGTH_ALIGNMENT=4");
 	}
 }
diff --git a/Engine/Plugins/Experimental/NNERuntimeRDG/Source/NNERuntimeRDG/Private/NNERuntimeRDGModel.cpp b/Engine/Plugins/Experimental/NNERuntimeRDG/Source/NNERuntimeRDG/Private/NNERuntimeRDGModel.cpp
--- a/Engine/Plugins/Experimental/NNERuntimeRDG/Source/NNERuntimeRDG/Private/NNERuntimeRDGModel.cpp
+++ b/Engine/Plugins/Experimental/NNERuntimeRDG/Source/NNERuntimeRDG/Private/NNERuntimeRDGModel.cpp
@@ -205,10 +205,15 @@ FModelInstanceRDG::ESetInputTensorShapesStatus FModelInstanceRDG::SetInputTensor
 	return ESetInputTensorShapesStatus::Ok;
 }
 
-FRDGBufferDesc CreateRDGBufferDescForTensorRDG(const FTensorRDG& Tensor)
+FRDGBufferDesc FModelInstanceRDG::CreateRDGBufferDescForTensorRDG(const FTensorRDG& Tensor)
 {
-	FRDGBufferDesc Desc = FRDGBufferDesc::CreateBufferDesc(Tensor.GetElementByteSize(), Tensor.GetVolume());
+	const uint32 ElementByteSize = Tensor.GetElementByteSize();
+	const uint32 TotalByteCount = ElementByteSize * Tensor.GetVolume();
 
+	//Round up to next multiple of BUFFER_LENGTH_ALIGNMENT
+	const uint32 TargetByteCount = FMath::DivideAndRoundUp(TotalByteCount, (uint32)NNERUNTIMERDGHLSL_BUFFER_LENGTH_ALIGNMENT) * NNERUNTIMERDGHLSL_BUFFER_LENGTH_ALIGNMENT;
+	
+	const FRDGBufferDesc Desc = FRDGBufferDesc::CreateBufferDesc(ElementByteSize, TargetByteCount / ElementByteSize);
 	return Desc;
 }
 
@@ -231,14 +236,14 @@ FModelInstanceRDG::EEnqueueRDGStatus FModelInstanceRDG::EnqueueRDG(FRDGBuilder&
 	Res = SetTensors(RDGBuilder, InputTensorRDGs, InInputBindings);
 	if (Res != -1)
 	{
-		UE_LOG(LogNNE, Warning, TEXT("Invalid buffer (was nullptr) for input tensor binding at index %d"), Res);
+		UE_LOG(LogNNE, Warning, TEXT("Invalid buffer for input tensor binding at index %d"), Res);
 		return EEnqueueRDGStatus::Fail;
 	}
 
 	Res = SetTensors(RDGBuilder, OutputTensorRDGs, InOutputBindings);
 	if (Res != -1)
 	{
-		UE_LOG(LogNNE, Warning, TEXT("Invalid buffer (was nullptr) for output tensor binding at index %d"), Res);
+		UE_LOG(LogNNE, Warning, TEXT("Invalid buffer for output tensor binding at index %d"), Res);
 		return EEnqueueRDGStatus::Fail;
 	}
 
@@ -286,6 +291,12 @@ int32 FModelInstanceRDG::SetTensors(FRDGBuilder& GraphBuilder, FTensorRDGArray&
 		const NNE::FTensorBindingRDG& Binding = InBindings[Idx];
 		if (Binding.Buffer == nullptr)
 		{
+			UE_LOG(LogNNE, Warning, TEXT("nullptr buffer encountered."));
+			return Idx;
+		}
+		if (Binding.Buffer->GetSize() % NNERUNTIMERDGHLSL_BUFFER_LENGTH_ALIGNMENT != 0)
+		{
+			UE_LOG(LogNNE, Warning, TEXT("Buffer has size %i which is not a multiple of %i"), Binding.Buffer->GetSize(), NNERUNTIMERDGHLSL_BUFFER_LENGTH_ALIGNMENT);
 			return Idx;
 		}
 		TensorRDG.SetBuffer(Binding.Buffer);
diff --git a/Engine/Plugins/Experimental/NNERuntimeRDG/Source/NNERuntimeRDG/Private/NNERuntimeRDGModel.h b/Engine/Plugins/Experimental/NNERuntimeRDG/Source/NNERuntimeRDG/Private/NNERuntimeRDGModel.h
--- a/Engine/Plugins/Experimental/NNERuntimeRDG/Source/NNERuntimeRDG/Private/NNERuntimeRDGModel.h
+++ b/Engine/Plugins/Experimental/NNERuntimeRDG/Source/NNERuntimeRDG/Private/NNERuntimeRDGModel.h
@@ -24,6 +24,7 @@ namespace UE::NNERuntimeRDG::Private
 		virtual EEnqueueRDGStatus EnqueueRDG(FRDGBuilder& RDGBuilder, TConstArrayView<NNE::FTensorBindingRDG> InInputBindings, TConstArrayView<NNE::FTensorBindingRDG> InOutputBindings) override;
 		
 	protected:
+		static FRDGBufferDesc CreateRDGBufferDescForTensorRDG(const FTensorRDG& Tensor);
 		bool LoadModel(TConstArrayView<uint8> ModelData, FNNERuntimeFormat& Format, int32 GuidAndVersionSize);
 		int32 SetTensors(FRDGBuilder& GraphBuilder, FTensorRDGArray& InTensorRDGs, TConstArrayView<NNE::FTensorBindingRDG> InBindings);
 
diff --git a/Engine/Plugins/Experimental/NNERuntimeRDG/Source/NNERuntimeRDG/Private/NNERuntimeRDGModelHlsl.cpp b/Engine/Plugins/Experimental/NNERuntimeRDG/Source/NNERuntimeRDG/Private/NNERuntimeRDGModelHlsl.cpp
--- a/Engine/Plugins/Experimental/NNERuntimeRDG/Source/NNERuntimeRDG/Private/NNERuntimeRDGModelHlsl.cpp
+++ b/Engine/Plugins/Experimental/NNERuntimeRDG/Source/NNERuntimeRDG/Private/NNERuntimeRDGModelHlsl.cpp
@@ -304,11 +304,12 @@ int FModelInstance::PrepareTensorShapesAndData()
 	return 0;
 }
 
-namespace UploadHelper
+void FModelInstance::PrepareWeights()
 {
+	check(WeightsExternalRDGResources.IsEmpty());
 
-void EnqueueTensorUpload(TArray<TRefCountPtr<FRDGPooledBuffer>>& OutExternalRDGResources,
-	FTensorRDGArray& TensorToUploadRDGs, ERDGInitialDataFlags CopyDataFlag)
+	auto EnqueueTensorUpload = [](TArray<TRefCountPtr<FRDGPooledBuffer>>& OutExternalRDGResources,
+								  FTensorRDGArray& TensorToUploadRDGs, ERDGInitialDataFlags CopyDataFlag)
 	{
 		OutExternalRDGResources.Reset();
 		OutExternalRDGResources.SetNum(TensorToUploadRDGs.Num());
@@ -372,17 +373,10 @@ void EnqueueTensorUpload(TArray<TRefCountPtr<FRDGPooledBuffer>>& OutExternalRDGR
 		Signal->Wait();	// Wait for render thread to finish
 
 		FGenericPlatformProcess::ReturnSynchEventToPool(Signal);
-	}
-}
-
-bool FModelInstance::PrepareWeights()
-{
-	check(WeightsExternalRDGResources.IsEmpty());
+	};
 
 	// Data is not copied. A GPU sync will happens see EnqueueTensorUpload().
-	UploadHelper::EnqueueTensorUpload(WeightsExternalRDGResources, WeightTensorRDGs, ERDGInitialDataFlags::NoCopy);
-
-	return true;
+	EnqueueTensorUpload(WeightsExternalRDGResources, WeightTensorRDGs, ERDGInitialDataFlags::NoCopy);
 }
 
 TSharedPtr<NNE::IModelInstanceRDG> FModel::CreateModelInstanceRDG()
diff --git a/Engine/Plugins/Experimental/NNERuntimeRDG/Source/NNERuntimeRDG/Private/NNERuntimeRDGModelHlsl.h b/Engine/Plugins/Experimental/NNERuntimeRDG/Source/NNERuntimeRDG/Private/NNERuntimeRDGModelHlsl.h
--- a/Engine/Plugins/Experimental/NNERuntimeRDG/Source/NNERuntimeRDG/Private/NNERuntimeRDGModelHlsl.h
+++ b/Engine/Plugins/Experimental/NNERuntimeRDG/Source/NNERuntimeRDG/Private/NNERuntimeRDGModelHlsl.h
@@ -27,7 +27,7 @@ protected:
 	virtual bool PrepareModelRDG(FRDGBuilder& RDGBuilder) override;
 	virtual void AddDispatchOps_RenderThread(FRDGBuilder& GraphBuilder) override;
 
-	bool PrepareWeights();
+	void PrepareWeights();
 
 private:
 
