NNE - Fix Hlsl Slice missing implementation for Steps tensor

Problems:
* Current implementation of RDG Hlsl Slice doesn't support `Steps` tensor

Solutions:
* Add support

Tests:
* Win64 Standalone tests: NNEQA Slice
* Win64 Editor tests: NNEQA Slice, Slice helper unit tests

#rnx

#virtualized

diff --git a/Engine/Plugins/Experimental/NNERuntimeRDG/Shaders/Private/NNEHlslShaders/NNEHlslShadersSlice.usf b/Engine/Plugins/Experimental/NNERuntimeRDG/Shaders/Private/NNEHlslShaders/NNEHlslShadersSlice.usf
--- a/Engine/Plugins/Experimental/NNERuntimeRDG/Shaders/Private/NNEHlslShaders/NNEHlslShadersSlice.usf
+++ b/Engine/Plugins/Experimental/NNERuntimeRDG/Shaders/Private/NNEHlslShaders/NNEHlslShadersSlice.usf
@@ -12,6 +12,7 @@ uint ThreadCountX;
 #define INPUT_STRIDE 0
 #define OUTPUT_STRIDE 1
 #define INPUT_STARTS 2
+#define INPUT_STEPS 3
 
 [numthreads(THREADGROUP_SIZE_X, 1, 1)]
 void Slice(in const uint3 DispatchThreadID : SV_DispatchThreadID)
@@ -27,7 +28,7 @@ void Slice(in const uint3 DispatchThreadID : SV_DispatchThreadID)
 		{
 			uint OutDimIdx, R;
 			DivMod(Offset, TensorInfo[dim][OUTPUT_STRIDE], OutDimIdx, R);
-			uint InDimIdx = TensorInfo[dim][INPUT_STARTS] + OutDimIdx;
+			uint InDimIdx = (uint) ((int) TensorInfo[dim][INPUT_STARTS] + (int) OutDimIdx * asint(TensorInfo[dim][INPUT_STEPS]));
 			InputIndex += TensorInfo[dim][INPUT_STRIDE] * InDimIdx;
 			Offset = R;
 		}
diff --git a/Engine/Plugins/Experimental/NNERuntimeRDG/Source/NNERuntimeRDG/Internal/NNERuntimeRDGHelperSlice.h b/Engine/Plugins/Experimental/NNERuntimeRDG/Source/NNERuntimeRDG/Internal/NNERuntimeRDGHelperSlice.h
--- a/Engine/Plugins/Experimental/NNERuntimeRDG/Source/NNERuntimeRDG/Internal/NNERuntimeRDGHelperSlice.h
+++ b/Engine/Plugins/Experimental/NNERuntimeRDG/Source/NNERuntimeRDG/Internal/NNERuntimeRDGHelperSlice.h
@@ -8,5 +8,5 @@ namespace UE::NNE::Internal { class FTensor; }
 
 namespace UE::NNERuntimeRDG::Internal::CPUHelper::Slice
 {
-	void NNERUNTIMERDG_API Apply(const NNE::Internal::FTensor& InputTensor, NNE::Internal::FTensor& OutputTensor, TConstArrayView<int32> Starts);
+	void NNERUNTIMERDG_API Apply(const NNE::Internal::FTensor& InputTensor, NNE::Internal::FTensor& OutputTensor, TConstArrayView<int32> Starts, TConstArrayView<int32> Steps);
 } // UE::NNERuntimeRDG::Internal::CPUHelper::Slice
\ No newline at end of file
diff --git a/Engine/Plugins/Experimental/NNERuntimeRDG/Source/NNERuntimeRDG/Private/Helper/NNERuntimeRDGHelperSlice.cpp b/Engine/Plugins/Experimental/NNERuntimeRDG/Source/NNERuntimeRDG/Private/Helper/NNERuntimeRDGHelperSlice.cpp
--- a/Engine/Plugins/Experimental/NNERuntimeRDG/Source/NNERuntimeRDG/Private/Helper/NNERuntimeRDGHelperSlice.cpp
+++ b/Engine/Plugins/Experimental/NNERuntimeRDG/Source/NNERuntimeRDG/Private/Helper/NNERuntimeRDGHelperSlice.cpp
@@ -5,35 +5,36 @@
 
 namespace UE::NNERuntimeRDG::Internal::CPUHelper::Slice
 {
-	template<typename TData> void ApplyResolvedInputType(const NNE::Internal::FTensor& InputTensor, NNE::Internal::FTensor& OutputTensor, TConstArrayView<int32> Starts)
+	template<typename TData> void ApplyResolvedInputType(const NNE::Internal::FTensor& InputTensor, NNE::Internal::FTensor& OutputTensor, TConstArrayView<int32> Starts, TConstArrayView<int32> Steps)
 	{
 		check(InputTensor.HasPreparedData());
 		check(InputTensor.GetShape().Rank() == Starts.Num());
 		check(OutputTensor.GetShape().Rank() == Starts.Num());
+		check(InputTensor.GetShape().Rank() == Steps.Num());
 
 		TArray<TData> OutputData;
 		TConstArrayView<TData> InputData = InputTensor.GetPreparedData<TData>();
-		Private::TensorIdxIterator itOutput(OutputTensor.GetShape());
-		const Private::TensorIdxIterator itInput(InputTensor.GetShape());
+		Private::TensorIdxIterator ItOutput(OutputTensor.GetShape());
+		const Private::TensorIdxIterator ItInput(InputTensor.GetShape());
 
 		OutputData.SetNumUninitialized(OutputTensor.GetVolume());
 		do
 		{
-			TArray<uint32> CurInputPosition(itOutput.GetPositions());
-			for (int r = 0; r < CurInputPosition.Num(); ++r)
+			TArray<uint32> CurInputPosition(ItOutput.GetPositions());
+			for (int DimIdx = 0; DimIdx < CurInputPosition.Num(); ++DimIdx)
 			{
-				CurInputPosition[r] += Starts[r];
+				CurInputPosition[DimIdx] = Starts[DimIdx] + (int32) CurInputPosition[DimIdx] * Steps[DimIdx];
 			}
 
-			TData Value = InputData[itInput.GetIndexFromPosition(CurInputPosition)];
-			OutputData[itOutput.GetIndex()] = Value;
+			TData Value = InputData[ItInput.GetIndexFromPosition(CurInputPosition)];
+			OutputData[ItOutput.GetIndex()] = Value;
 
-		} while (itOutput.Advance());
+		} while (ItOutput.Advance());
 
 		OutputTensor.SetPreparedData<TData>(OutputData);
 	}
 
-	void Apply(const NNE::Internal::FTensor& InputTensor, NNE::Internal::FTensor& OutputTensor, TConstArrayView<int32> Starts)
+	void Apply(const NNE::Internal::FTensor& InputTensor, NNE::Internal::FTensor& OutputTensor, TConstArrayView<int32> Starts, TConstArrayView<int32> Steps)
 	{
 		static constexpr int32 MaxItemInOutputTensor = NNE::FTensorShape::MaxRank * 2;
 
@@ -50,15 +51,15 @@ namespace UE::NNERuntimeRDG::Internal::CPUHelper::Slice
 		switch (InputTensor.GetDataType())
 		{
 		case ENNETensorDataType::Int32:
-			ApplyResolvedInputType<int32>(InputTensor, OutputTensor, Starts);
+			ApplyResolvedInputType<int32>(InputTensor, OutputTensor, Starts, Steps);
 			break;
 
 		case ENNETensorDataType::Int64:
-			ApplyResolvedInputType<int64>(InputTensor, OutputTensor, Starts);
+			ApplyResolvedInputType<int64>(InputTensor, OutputTensor, Starts, Steps);
 			break;
 
 		case ENNETensorDataType::Float:
-			ApplyResolvedInputType<float>(InputTensor, OutputTensor, Starts);
+			ApplyResolvedInputType<float>(InputTensor, OutputTensor, Starts, Steps);
 			break;
 		}
 	}
diff --git a/Engine/Plugins/Experimental/NNERuntimeRDG/Source/NNERuntimeRDG/Private/Hlsl/NNERuntimeRDGSlice.cpp b/Engine/Plugins/Experimental/NNERuntimeRDG/Source/NNERuntimeRDG/Private/Hlsl/NNERuntimeRDGSlice.cpp
--- a/Engine/Plugins/Experimental/NNERuntimeRDG/Source/NNERuntimeRDG/Private/Hlsl/NNERuntimeRDGSlice.cpp
+++ b/Engine/Plugins/Experimental/NNERuntimeRDG/Source/NNERuntimeRDG/Private/Hlsl/NNERuntimeRDGSlice.cpp
@@ -26,76 +26,118 @@ namespace UE::NNERuntimeRDG::Private::Hlsl
 
 	private:
 
-		mutable TArray<int32, TInlineAllocator<NNE::FTensorShape::MaxRank>> AxesAttr;
-		mutable TArray<int32, TInlineAllocator<NNE::FTensorShape::MaxRank>> EndsAttr;
-		mutable TArray<int32, TInlineAllocator<NNE::FTensorShape::MaxRank>> StartsAttr;
+		TArray<int32, TInlineAllocator<NNE::FTensorShape::MaxRank>> AxesAttr;
+		TArray<int32, TInlineAllocator<NNE::FTensorShape::MaxRank>> EndsAttr;
+		TArray<int32, TInlineAllocator<NNE::FTensorShape::MaxRank>> StartsAttr;
+		TArray<int32, TInlineAllocator<NNE::FTensorShape::MaxRank>> StepsAttr;
 
-		mutable TArray<int32, TInlineAllocator<NNE::FTensorShape::MaxRank>> Start;
-		mutable TArray<int32, TInlineAllocator<NNE::FTensorShape::MaxRank>> End;
+		TArray<int32, TInlineAllocator<NNE::FTensorShape::MaxRank>> Start;
+		TArray<int32, TInlineAllocator<NNE::FTensorShape::MaxRank>> End;
+		TArray<int32, TInlineAllocator<NNE::FTensorShape::MaxRank>> Step;
 
-		bool TryGetAttributesFromConstantTensors(TConstArrayView<NNE::Internal::FTensorRef> InputTensors) const
+		bool TryGetAttributesFromConstantTensors(TConstArrayView<NNE::Internal::FTensorRef> InputTensors)
 		{
 			check(InputTensors.Num() >= 3 && InputTensors.Num() <= 5);
 
-			if (!OperatorHelper::GetInt32ArrayFromConstTensor(StartsAttr, InputTensors[1]))
-			{
-				UE_LOG(LogNNE, Warning, TEXT("Error: Slice op 'Starts' input at index 1 is only supported as a constant integer tensor but it is not."));
-				return false;
-			}
-			if (!OperatorHelper::GetInt32ArrayFromConstTensor(EndsAttr, InputTensors[2]))
-			{
-				UE_LOG(LogNNE, Warning, TEXT("Error: Slice op 'Ends' input at index 2 is only supported as a constant integer tensor but it is not."));
-				return false;
-			}
 			if (InputTensors.Num() >= 4)
 			{
 				if (!OperatorHelper::GetInt32ArrayFromConstTensor(AxesAttr, InputTensors[3]))
 				{
-					UE_LOG(LogNNE, Warning, TEXT("Error: Slice op 'Axes' input at index 3 is only supported as a constant integer tensor but it is not."));
+					UE_LOG(LogNNE, Warning, TEXT("Error: Slice op 'Axes' input tensor (%s) is only supported as a constant integer tensor but it is not."), *InputTensors[3]->GetName());
+					return false;
+				}
+				if(AxesAttr.Num() > InputTensors[0]->GetShape().Rank() || AxesAttr.Num() < 1)
+				{
+					UE_LOG(LogNNE, Warning, TEXT("Error: Slice op 'Axes' input tensor (%s) contains %d elements but input rank is %d."), *InputTensors[3]->GetName(), AxesAttr.Num(), InputTensors[0]->GetShape().Rank());
 					return false;
 				}
 			}
 			else
 			{
-				for (int32 i = 0; i < StartsAttr.Num(); ++i)
+				for (int32 Idx = 0; Idx < InputTensors[0]->GetShape().Rank(); ++Idx)
 				{
-					AxesAttr.Add(i);
+					AxesAttr.Add(Idx);
 				}
 			}
+
+			const int32 NumAxes = AxesAttr.Num();
+
+			if (!OperatorHelper::GetInt32ArrayFromConstTensor(StartsAttr, InputTensors[1]))
+			{
+				UE_LOG(LogNNE, Warning, TEXT("Error: Slice op 'Starts' input tensor (%s) is only supported as a constant integer tensor but it is not."), *InputTensors[1]->GetName());
+				return false;
+			}
+			if(StartsAttr.Num() != NumAxes)
+			{
+				UE_LOG(LogNNE, Warning, TEXT("Error: Slice op 'Starts' input tensor (%s) contains %d elements but number of axes is %d."), *InputTensors[1]->GetName(), StartsAttr.Num(), NumAxes);
+				return false;
+			}
+
+			if (!OperatorHelper::GetInt32ArrayFromConstTensor(EndsAttr, InputTensors[2]))
+			{
+				UE_LOG(LogNNE, Warning, TEXT("Error: Slice op 'Ends' input tensor (%s) is only supported as a constant integer tensor but it is not."), *InputTensors[2]->GetName());
+				return false;
+			}
+			if(EndsAttr.Num() != NumAxes)
+			{
+				UE_LOG(LogNNE, Warning, TEXT("Error: Slice op 'Ends' input tensor (%s) contains %d elements but number of axes is %d."), *InputTensors[2]->GetName(), EndsAttr.Num(), NumAxes);
+				return false;
+			}
+
+			
+
 			if (InputTensors.Num() == 5)
 			{
-				TArray<int32, TInlineAllocator<NNE::FTensorShape::MaxRank>> StepsAttr;
-				OperatorHelper::GetInt32ArrayFromConstTensor(StepsAttr, InputTensors[4]);
-				for (int32 Value : StepsAttr)
+				if (!OperatorHelper::GetInt32ArrayFromConstTensor(StepsAttr, InputTensors[4]))
+				{
+					UE_LOG(LogNNE, Warning, TEXT("Error: Slice op 'Steps' input tensor (%s) is only supported as a constant integer tensor but it is not."), *InputTensors[4]->GetName());
+					return false;
+				}
+				if(StepsAttr.Num() != NumAxes)
 				{
-					if (Value != 1)
+					UE_LOG(LogNNE, Warning, TEXT("Error: Slice op 'Steps' input tensor (%s) contains %d elements but number of axes is %d."), *InputTensors[4]->GetName(), StepsAttr.Num(), NumAxes);
+					return false;
+				}
+				for(const int32 Value : StepsAttr)
+				{
+					if(Value == 0)
 					{
-						UE_LOG(LogNNE, Warning, TEXT("Error: Slice op 'Steps' optional input at index 4 is only supported as a constant tensor will all values as 1."));
+						UE_LOG(LogNNE, Warning, TEXT("Error: Slice op 'Steps' tensor (%s) can only contain non-0 integers."), *InputTensors[4]->GetName());
 						return false;
 					}
 				}
 			}
+			else
+			{
+				StepsAttr.Init(1, AxesAttr.Num()); // Default for steps is all 1s
+			}
 
 			return true;
 		}
 
-		void ComputeStartAndEndFromInputShape(TConstArrayView<uint32> InputShapeData) const
+		void ComputeStartAndEndFromInputShape(TConstArrayView<uint32> InputShapeData)
 		{
-			check(AxesAttr.Num() == EndsAttr.Num());
+			const int32 InputRank = InputShapeData.Num();
+
+			check(AxesAttr.Num() <= InputRank);
 			check(AxesAttr.Num() == StartsAttr.Num());
+			check(AxesAttr.Num() == EndsAttr.Num());
+			check(AxesAttr.Num() == StepsAttr.Num());
 
-			const int32 InputRank = InputShapeData.Num();
-			TArray<int32, TInlineAllocator<NNE::FTensorShape::MaxRank>> Axes(AxesAttr);
-			TArray<int32, TInlineAllocator<NNE::FTensorShape::MaxRank>> Ends(EndsAttr);
-			TArray<int32, TInlineAllocator<NNE::FTensorShape::MaxRank>> Starts(StartsAttr);
+			TArray<int32, TInlineAllocator<NNE::FTensorShape::MaxRank>>& Axes(AxesAttr);
+			TArray<int32, TInlineAllocator<NNE::FTensorShape::MaxRank>>& Starts(StartsAttr);
+			TArray<int32, TInlineAllocator<NNE::FTensorShape::MaxRank>>& Ends(EndsAttr);
+			TArray<int32, TInlineAllocator<NNE::FTensorShape::MaxRank>>& Steps(StepsAttr);
 
 			//see [REDACTED_PROFILE_URL]/onnx/blob/main/docs/Operators.md#slice for algorithm
 			Start.SetNum(InputRank);
 			End.SetNum(InputRank);
+			Step.SetNum(InputRank);
 			for (int32 i = 0; i < InputRank; ++i)
 			{
 				Start[i] = 0;
 				End[i] = InputShapeData[i];
+				Step[i] = 1;
 			}
 			for (int32 i = 0; i < Axes.Num(); ++i)
 			{
@@ -120,8 +162,17 @@ namespace UE::NNERuntimeRDG::Private::Hlsl
 			}
 			for (int32 i = 0; i < Axes.Num(); ++i)
 			{
-				Start[Axes[i]] = FMath::Clamp(Starts[i], 0, InputShapeData[Axes[i]]);
-				End[Axes[i]] = FMath::Clamp(Ends[i], 0, InputShapeData[Axes[i]]);
+				if(Steps[i] > 0)
+				{
+					Start[Axes[i]] = FMath::Clamp(Starts[i], 0, InputShapeData[Axes[i]]);
+					End[Axes[i]] = FMath::Clamp(Ends[i], 0, InputShapeData[Axes[i]]);
+				}
+				else
+				{
+					Start[Axes[i]] = FMath::Clamp(Starts[i], 0, InputShapeData[Axes[i]] - 1);
+					End[Axes[i]] = FMath::Clamp(Ends[i], -1, InputShapeData[Axes[i]] - 1);
+				}
+				Step[Axes[i]] = Steps[i];
 			}
 		}
 
@@ -152,20 +203,27 @@ namespace UE::NNERuntimeRDG::Private::Hlsl
 			TArray<uint32> OutputShapeData;
 
 			OutputShapeData.Reserve(InputRank);
-			for (int32 i = 0; i < InputRank; ++i)
+			for (int32 Idx = 0; Idx < InputRank; ++Idx)
 			{
-				OutputShapeData.Add(End[i] - Start[i]);
+				int32 RangeSize = Step[Idx] > 0 ? End[Idx] - Start[Idx] : Start[Idx] - End[Idx];
+				if(RangeSize < 1)
+				{
+					UE_LOG(LogNNE, Warning, TEXT("Error: Slice op start (tensor: %s) and end (tensor: %s) indices are incompatible with step direction for dimension %d."), *InputTensors[1]->GetName(), *InputTensors[2]->GetName(), Idx);
+					return -1;
+				}
+				uint32 OutDimSize = FMath::DivideAndRoundUp(RangeSize, FMath::Abs(Step[Idx]));
+				OutputShapeData.Add(OutDimSize);
 			}
 			
 			NNE::FTensorShape OutputShape = NNE::FTensorShape::Make(OutputShapeData);
 
 			OutputTensors[0]->SetShape(OutputShape);
 			
-			Internal::CPUHelper::Slice::Apply(*InputTensors[0], *OutputTensors[0], Start);
+			Internal::CPUHelper::Slice::Apply(*InputTensors[0], *OutputTensors[0], Start, Step);
 
-			if (OutputTensors[0]->GetDataType() != ENNETensorDataType::Float && !OutputTensors[0]->HasPreparedData())
+			if (InputTensors[0]->HasPreparedData() && !OutputTensors[0]->HasPreparedData())
 			{
-				UE_LOG(LogNNE, Warning, TEXT("Error: Slice op output tensor could not be made constant nor it was of float type. Only floats are supported at the moment on the HLSL compute path."));
+				UE_LOG(LogNNE, Warning, TEXT("Error: Slice op output tensor (%s) could not be constant-folded from input."), *OutputTensors[0]->GetName());
 				return -1;
 			}
 
@@ -197,11 +255,14 @@ namespace UE::NNERuntimeRDG::Private::Hlsl
 					UE_LOG(LogNNE, Warning, TEXT("Slice: Starts, Ends and Axes must be of the same size."));
 					return false;
 				}
+
+				StepsAttr.Init(1, AxesAttr.Num()); // Default for steps is all 1s
 			}
 			else
 			{
 				check(InputTensorDescs.Num() >= 3 && InputTensorDescs.Num() <= 5);
 			}
+
 			
 			return true;
 		}
@@ -231,16 +292,10 @@ namespace UE::NNERuntimeRDG::Private::Hlsl
 			FillTensorStrideShaderParameters(Output, Params->TensorInfo, 1);
 			static_assert(NNE::FTensorShape::MaxRank <= NXRT_TENSORSTRIDEINFO_MAX_NUM_DIMENSIONS);
 			check(Start.Num() == Input.GetShape().Rank());
-			for (int32 i = 0; i < NXRT_TENSORSTRIDEINFO_MAX_NUM_DIMENSIONS; ++i)
+			for (int32 DimIdx = 0; DimIdx < Input.GetShape().Rank(); ++DimIdx)
 			{
-				if (i < Start.Num())
-				{
-					Params->TensorInfo[i][2] = static_cast<uint32>(Start[i]);
-				}
-				else
-				{
-					Params->TensorInfo[i][2] = 0;
-				}
+				Params->TensorInfo[DimIdx][/*SlotIdx*/ 2] = static_cast<uint32>(Start[DimIdx]);
+				Params->TensorInfo[DimIdx][/*SlotIdx*/ 3] = BitCast<uint32>(Step[DimIdx]);
 			}
 
 			FSliceCS::FPermutationDomain PermutationVector;
