NNE - Extend Hlsl Conv to bigger filters

Problems:
* Hlsl Conv is limited to 3x3 filters due to a limitation on the number of reads per thread
* Instead of reporting an error to the user, it crashes on a check()

Solutions:
* Increase the number of reads per thread limit and introduce an algorithm that selects the best thread group size that doesn't violate shared memory limits
* Report to the user and fail gracefully when a configuration is not supported

Tests:
* Win64 Editor tests: All

#rnx

#virtualized

diff --git a/Engine/Plugins/Experimental/NNERuntimeRDG/Source/NNEHlslShaders/Internal/NNEHlslShadersConvCS.h b/Engine/Plugins/Experimental/NNERuntimeRDG/Source/NNEHlslShaders/Internal/NNEHlslShadersConvCS.h
--- a/Engine/Plugins/Experimental/NNERuntimeRDG/Source/NNEHlslShaders/Internal/NNEHlslShadersConvCS.h
+++ b/Engine/Plugins/Experimental/NNERuntimeRDG/Source/NNEHlslShaders/Internal/NNEHlslShadersConvCS.h
@@ -36,7 +36,7 @@ namespace UE::NNEHlslShaders::Internal
 	public:
 		static const int32 MAX_NUM_DIMENSIONS{4};
 		static const int32 MIN_NUM_READS_PER_THREAD_POW2{1};
-		static const int32 MAX_NUM_READS_PER_THREAD_POW2{3};
+		static const int32 MAX_NUM_READS_PER_THREAD_POW2{4};
 	};
 
 	class NNEHLSLSHADERS_API FConvCS : public FHlslShaderBase
@@ -78,6 +78,7 @@ namespace UE::NNEHlslShaders::Internal
 			SHADER_PARAMETER(float, GroupsDivM)
 		END_SHADER_PARAMETER_STRUCT()
 
+		static bool ShouldCompilePermutation(const FGlobalShaderPermutationParameters& InParameters);
 		static void ModifyCompilationEnvironment(const FGlobalShaderPermutationParameters& InParameters, FShaderCompilerEnvironment& OutEnvironment);
 
 	public:
@@ -106,6 +107,7 @@ namespace UE::NNEHlslShaders::Internal
 		*/
 		static FIntVector GetGroupCount(TArrayView<const int32> YShape, TArrayView<const int32> GroupShape);
 
+		static EConvGroupSize GetBiggestCompatibleGroupSize(TArrayView<const uint32> WShape, TArrayView<const int32> Dilations, TArrayView<const int32> Strides);
 		static EConvGroupSize GetMinimalGroupSize(TArrayView<const int32> WShape);
 
 		static TArray<int32> GetPadding(TArrayView<const uint32> XShape, TArrayView<const uint32> WShape, EConvAutoPad AutoPad, TArrayView<const int32> Dilations, TArrayView<const int32> Strides, TArrayView<const int32> Pads);
diff --git a/Engine/Plugins/Experimental/NNERuntimeRDG/Source/NNEHlslShaders/Private/NNEHlslShadersConvCS.cpp b/Engine/Plugins/Experimental/NNERuntimeRDG/Source/NNEHlslShaders/Private/NNEHlslShadersConvCS.cpp
--- a/Engine/Plugins/Experimental/NNERuntimeRDG/Source/NNEHlslShaders/Private/NNEHlslShadersConvCS.cpp
+++ b/Engine/Plugins/Experimental/NNERuntimeRDG/Source/NNEHlslShaders/Private/NNEHlslShadersConvCS.cpp
@@ -2,6 +2,7 @@
 
 #include "NNEHlslShadersConvCS.h"
 #include "NNE.h"
+#include "RHIGlobals.h"
 
 namespace UE::NNEHlslShaders::Internal
 {
@@ -40,7 +41,7 @@ namespace UE::NNEHlslShaders::Internal
 				NumThreadsPerGroup = 512;
 				break;
 			default:
-				NumThreadsPerGroup = 128;
+				check(false);
 				break;
 		}
 		check(FMath::Log2((float)NumThreadsPerGroup) == FMath::Floor(FMath::Log2((float)NumThreadsPerGroup)));
@@ -62,14 +63,70 @@ namespace UE::NNEHlslShaders::Internal
 		return Result;
 	}
 
+	uint32 GetTotalSharedMemoryUsed(uint32 ThreadGroupSize, uint32 Log2NumReadsPerThread)
+	{
+		return ThreadGroupSize * (1 + (1 << Log2NumReadsPerThread)) * sizeof(float);
+	}
+
 	} // namespace ConvUtils
 
+	bool FConvCS::ShouldCompilePermutation(const FGlobalShaderPermutationParameters& InParameters)
+	{
+		if (!FHlslShaderBase::ShouldCompilePermutation(InParameters))
+		{
+			return false;
+		}
+
+		FPermutationDomain PermutationVector(InParameters.PermutationId);
+
+		int Log2NumReadsPerThread = PermutationVector.Get<FConvNumReadsPerThread>();
+		int32 ConvGroupSize = ConvUtils::GetNumThreadsPerGroup(PermutationVector.Get<FConvGroupSize>());
+
+		if(ConvUtils::GetTotalSharedMemoryUsed(ConvGroupSize, Log2NumReadsPerThread) > GetMaxComputeSharedMemory())
+		{
+			return false;
+		}
+		
+		return true;
+	}
+
 	void FConvCS::ModifyCompilationEnvironment(const FGlobalShaderPermutationParameters& InParameters, FShaderCompilerEnvironment& OutEnvironment)
 	{
 		FGlobalShader::ModifyCompilationEnvironment(InParameters, OutEnvironment);
 		OutEnvironment.SetDefine(TEXT("MAX_NUM_DIMENSIONS"), FConvConstants::MAX_NUM_DIMENSIONS);
 	}
 
+	// Returns the biggest thread group size that allows to keep in shared memory the portion of the input tensor and the filter necessary to compute the portion of output tended by it (the thread group).
+	// If no suitable thread group size was found, EConvGroupSize::MAX is returned.
+	EConvGroupSize FConvCS::GetBiggestCompatibleGroupSize(TArrayView<const uint32> WShape, TArrayView<const int32> Dilations, TArrayView<const int32> Strides)
+	{
+		check(WShape.Num() > 2);
+		check(Dilations.Num() == 0 || Dilations.Num() == WShape.Num() - 2);
+		check(Strides.Num() == 0 || Strides.Num() == WShape.Num() - 2);
+		
+		for(EConvGroupSize TargetGroupSize = (EConvGroupSize)((uint8)EConvGroupSize::MAX - 1); (uint8)TargetGroupSize > 0; TargetGroupSize = (EConvGroupSize)((uint8)TargetGroupSize - 1))
+		{
+			const int32 NumDimensions = WShape.Num() - 2;
+			const TArray<int32> GroupShape = GetGroupShape(TargetGroupSize, NumDimensions);
+			TArray<int32> TargetXBlockShape = ConvUtils::GetXBlockShape(GroupShape, WShape, Dilations, Strides);
+			int32 XWindowSize = 1;
+			for (int32 Idx = 0; Idx < NumDimensions; Idx++)
+			{
+				XWindowSize *= TargetXBlockShape[Idx];
+			}
+			int32 NumReadsPerThread = FMath::DivideAndRoundUp(XWindowSize, ConvUtils::GetNumThreadsPerGroup(TargetGroupSize));
+			int32 Log2NumReadsPerThread = FMath::Max(FMath::RoundToPositiveInfinity(FMath::Log2((float)NumReadsPerThread)), FConvConstants::MIN_NUM_READS_PER_THREAD_POW2);
+			uint32 TotSharedMemoryUsed = ConvUtils::GetTotalSharedMemoryUsed(ConvUtils::GetNumThreadsPerGroup(TargetGroupSize), Log2NumReadsPerThread);
+
+			if(TotSharedMemoryUsed <= GetMaxComputeSharedMemory() && Log2NumReadsPerThread <= FConvConstants::MAX_NUM_READS_PER_THREAD_POW2)
+			{
+				return TargetGroupSize;
+			}
+		}
+
+		return EConvGroupSize::MAX;
+	}
+
 	TArray<int32> FConvCS::GetPadding(TArrayView<const uint32> XShape, TArrayView<const uint32> WShape, EConvAutoPad AutoPad, TArrayView<const int32> Dilations, TArrayView<const int32> Strides, TArrayView<const int32> Pads)
 	{
 		check(XShape.Num() > 2);
@@ -145,7 +202,7 @@ namespace UE::NNEHlslShaders::Internal
 		check(Dilations.Num() == 0 || Dilations.Num() == WShape.Num() - 2);
 		check(Strides.Num() == 0 || Strides.Num() == WShape.Num() - 2);
 		check(Pads.Num() == 0 || Pads.Num() == 2 * (WShape.Num() - 2));
-		check(GetNumReadsPerThread(GroupSize, WShape, Dilations, Strides) >= 0)
+		check(GroupSize != EConvGroupSize::MAX);
 		check(WShape[0] > 0)
 		check(WShape[1] > 0)
 
@@ -271,7 +328,7 @@ namespace UE::NNEHlslShaders::Internal
 
 	EConvGroupSize FConvCS::GetMinimalGroupSize(TArrayView<const int32> WShape)
 	{
-		int32 NumDimensions = WShape.Num() - 2;
+		const int32 NumDimensions = WShape.Num() - 2;
 		int32 WChannelSize = 1;
 		for (int32 i = 0; i < NumDimensions; i++)
 		{
diff --git a/Engine/Plugins/Experimental/NNERuntimeRDG/Source/NNERuntimeRDG/Private/Hlsl/NNERuntimeRDGConv.cpp b/Engine/Plugins/Experimental/NNERuntimeRDG/Source/NNERuntimeRDG/Private/Hlsl/NNERuntimeRDGConv.cpp
--- a/Engine/Plugins/Experimental/NNERuntimeRDG/Source/NNERuntimeRDG/Private/Hlsl/NNERuntimeRDGConv.cpp
+++ b/Engine/Plugins/Experimental/NNERuntimeRDG/Source/NNERuntimeRDG/Private/Hlsl/NNERuntimeRDGConv.cpp
@@ -17,6 +17,7 @@ namespace UE::NNERuntimeRDG::Private::Hlsl
 	DECLARE_GPU_STAT_NAMED(FNNEOperatorConvMatmul, TEXT("NNE.Operator.Hlsl.Conv.Matmul"));
 
 	using EConvAutoPad = UE::NNEHlslShaders::Internal::EConvAutoPad;
+	using EConvGroupSize = UE::NNEHlslShaders::Internal::EConvGroupSize;
 
 	/**
 	 * Convolution operator implementation
@@ -45,17 +46,27 @@ namespace UE::NNERuntimeRDG::Private::Hlsl
 		TArray<int32> Strides;
 		bool bAreWeightsTransposed = false;
 		EPixelFormat BufferPixelFormat;
+		EConvGroupSize GroupSize = EConvGroupSize::MAX;
 
 	public:
 
 		virtual int PrepareOutputs(TConstArrayView<NNE::Internal::FTensorRef> InputTensors, TArrayView<NNE::Internal::FTensorRef> OutputTensors) override
 		{
+			using namespace UE::NNEHlslShaders::Internal;
+
 			check(InputTensors.Num() >= 2 && InputTensors.Num() <= 3);
 			check(OutputTensors.Num() == 1);
 
 			const NNE::FTensorShape& Input = InputTensors[0]->GetShape();
 			const NNE::FTensorShape& Weights = InputTensors[1]->GetShape();
 
+			GroupSize = FConvCS::GetBiggestCompatibleGroupSize(Weights.GetData(), Dilations, Strides);
+			if(GroupSize == EConvGroupSize::MAX)
+			{
+				UE_LOG(LogNNE, Warning, TEXT("Hlsl Conv: kernel size, strides, dilations combination is not supported. Kernel tensor: %s."), *InputTensors[1]->GetName());
+				return -1;
+			}
+
 			const TArray<int32> OutputShapeData = NNEHlslShaders::Internal::FConvCS::GetOutputShape(Input.GetData(), Weights.GetData(), AutoPad, Dilations, Strides, Pads);
 			const NNE::FSymbolicTensorShape OutputShape = NNE::FSymbolicTensorShape::Make(OutputShapeData);
 			
@@ -121,7 +132,7 @@ namespace UE::NNERuntimeRDG::Private::Hlsl
 			using namespace UE::NNEHlslShaders::Internal;
 
 			constexpr EConvAlgorithm Algorithm = EConvAlgorithm::SharedMemory;
-			constexpr EConvGroupSize GroupSize = EConvGroupSize::Size256;
+			
 			bool HasBias = Bias != nullptr;
 			TArray<int32> OutputShape = FConvCS::GetOutputShape(Input.GetShape().GetData(), Weights.GetShape().GetData(), AutoPad, Dilations, Strides, Pads);
 
