NNE - Fix support for scalar indices tensor in Hlsl Gather

Problems:
* Hlsl Gather hits a check when providing scalar (0-dimensional) indices tensors

Solutions:
* Extend Hlsl Gather to support scalar indices tensor

Tests:
* Win64 NNEQA Editor: all
* Win64 NNEQA Standalone: all

#rnx

#virtualized

diff --git a/Engine/Plugins/Experimental/NNERuntimeRDG/Shaders/Private/NNEHlslShaders/NNEHlslShadersGather.usf b/Engine/Plugins/Experimental/NNERuntimeRDG/Shaders/Private/NNEHlslShaders/NNEHlslShadersGather.usf
--- a/Engine/Plugins/Experimental/NNERuntimeRDG/Shaders/Private/NNEHlslShaders/NNEHlslShadersGather.usf
+++ b/Engine/Plugins/Experimental/NNERuntimeRDG/Shaders/Private/NNEHlslShaders/NNEHlslShadersGather.usf
@@ -45,7 +45,7 @@ void Gather(in const uint3 DispatchThreadID : SV_DispatchThreadID)
 	}
 
 	// Split the output indices into data and indices indices 
-	int DataIndices[NUM_OUTPUT_DIMENSIONS];
+	int DataIndices[NUM_OUTPUT_DIMENSIONS + 1];
 	int IndicesIndices[NUM_OUTPUT_DIMENSIONS];
 	for (int i = 0; i < Axis; i++)
 	{
diff --git a/Engine/Plugins/Experimental/NNERuntimeRDG/Source/NNERuntimeRDG/Private/Hlsl/NNERuntimeRDGGather.cpp b/Engine/Plugins/Experimental/NNERuntimeRDG/Source/NNERuntimeRDG/Private/Hlsl/NNERuntimeRDGGather.cpp
--- a/Engine/Plugins/Experimental/NNERuntimeRDG/Source/NNERuntimeRDG/Private/Hlsl/NNERuntimeRDGGather.cpp
+++ b/Engine/Plugins/Experimental/NNERuntimeRDG/Source/NNERuntimeRDG/Private/Hlsl/NNERuntimeRDGGather.cpp
@@ -73,24 +73,10 @@ namespace UE::NNERuntimeRDG::Private::Hlsl
 			check(InputTensorDescs.Num() == 2)
 			check(OutputTensorDescs.Num() == 1)
 
-			const int32 MaxNumDimensions = NNEHlslShaders::Internal::FGatherConstants::MAX_NUM_DIMENSIONS;
-
 			const NNE::FTensorDesc& Data = InputTensorDescs[0];
 			const NNE::FTensorDesc& Indices = InputTensorDescs[1];
 			const NNE::FTensorDesc& Output = OutputTensorDescs[0];
 
-			if (Output.GetShape().Rank() > MaxNumDimensions)
-			{
-				UE_LOG(LogNNE, Warning, TEXT("Gather first input should be of rank %d or less but is %d"), MaxNumDimensions, Output.GetShape().Rank());
-				return false;
-			}
-
-			if ((Data.GetShape().Rank() + Indices.GetShape().Rank() - 1) > MaxNumDimensions)
-			{
-				UE_LOG(LogNNE, Warning, TEXT("Gather sum of input 0 and 1 ranks -1 should be less than %d"), MaxNumDimensions);
-				return false;
-			}
-
 			Axis = Attributes.GetValueOrDefault(TEXT("axis"), Axis);
 			if (Axis >= Data.GetShape().Rank())
 			{
@@ -118,8 +104,7 @@ namespace UE::NNERuntimeRDG::Private::Hlsl
 			check(OutputTensors[0] != nullptr);
 			check(OutputTensors[0]->GetShape().Rank() <= FGatherConstants::MAX_NUM_DIMENSIONS)
 			check(InputTensors[0]->GetShape().Rank() > 0)
-			check(InputTensors[1]->GetShape().Rank() > 0)
-			check(InputTensors[0]->GetShape().Rank() + (InputTensors[1]->GetShape().Rank() - 1) <= FGatherConstants::MAX_NUM_DIMENSIONS)
+			check(InputTensors[1]->GetShape().Rank() + (InputTensors[0]->GetShape().Rank() - 1) <= FGatherConstants::MAX_NUM_DIMENSIONS)
 
 			const FTensorRDG& Data = *InputTensors[0];
 			const FTensorRDG& Indices = *InputTensors[1];
@@ -175,7 +160,25 @@ namespace UE::NNERuntimeRDG::Private::Hlsl
 		InputValidator.AddRequired(1);
 		bIsValid &= InputValidator.Validate(InputTypes);
 
-		return bIsValid;
+		if(!bIsValid)
+		{
+			return false;
+		}
+
+		if(InputShapes[0].Rank() < 1)
+		{
+			UE_LOG(LogNNE, Warning, TEXT("Hlsl Gather: input tensor must have rank >= 1."));
+			return false;
+		}
+
+		const int32 OutputRank = InputShapes[1].Rank() + (InputShapes[0].Rank() - 1);
+		if(OutputRank > NNEHlslShaders::Internal::FGatherConstants::MAX_NUM_DIMENSIONS)
+		{
+			UE_LOG(LogNNE, Warning, TEXT("Hlsl Gather: output tensor has rank %d higher than maximum supported: %d."), OutputRank, NNEHlslShaders::Internal::FGatherConstants::MAX_NUM_DIMENSIONS);
+			return false;
+		}
+
+		return true;
 	}
 
 	FOperatorHlsl* CreateGatherOperator()
