NNE - Fix NaNs produced by erf() on certain platforms

Problems:
* Hlsl erf (error function) is implemented via tanh() intrinsics. On certain platforms, for large absolute values of the input, the implementation of tanh() produces NaNs.

Solutions:
* Branch for large enough absolute values of the input and return the asymptotic value instead.

Tests:
* Win64 NNEQA Editor: all
* Win64 NNEQA Standalone: all
* Win64 NNEQA Consoles: tested via MHA model

#rnx

#virtualized

diff --git a/Engine/Plugins/Experimental/NNERuntimeRDG/Shaders/Private/NNEHlslShaders/NNEHlslShadersElementWiseUnary.usf b/Engine/Plugins/Experimental/NNERuntimeRDG/Shaders/Private/NNEHlslShaders/NNEHlslShadersElementWiseUnary.usf
--- a/Engine/Plugins/Experimental/NNERuntimeRDG/Shaders/Private/NNEHlslShaders/NNEHlslShadersElementWiseUnary.usf
+++ b/Engine/Plugins/Experimental/NNERuntimeRDG/Shaders/Private/NNEHlslShaders/NNEHlslShadersElementWiseUnary.usf
@@ -96,6 +96,12 @@ float elu(float x)
 //see [REDACTED_PROFILE_URL]/onnx/blob/main/docs/Operators.md#erf
 float erf(float x)
 {
+	// On some platforms tanh() gives NaN for large inputs, 
+	// thus better directly return the asymptotic value (+1 or -1) for abs(x) >= 4
+	if(abs(x) >= 4.0f)
+	{
+		return asfloat((asuint(x) & 0x80000000) ^ asuint(1.0f));
+	}
 	//https://aapt.scitation.org/doi/abs/10.1119/1.15018?journalCode=ajp
 	float a = 167.0f/148.0f;
 	float b = 11.0f/109.0f;
