NNE - Settle hardware capability support for NNERuntimeHlsl

Problems:
* Hardware capability support is not enforced explicitly

Solutions:
* Put platform support checks in shader base's `ShouldCompilePermutation()` to skip compiling NNE shaders when target platform doesn't support them.
* Put current system support checks in `CanCreateModelRDG()`, disallowing model creation on unsupported systems
* Introduce WITH_NNE_RUNTIME_HLSL define enabling registration of the runtime only for supported systems (non-editor builds).

Tests:
* Win64 Editor PCD3D_SM6: manually checked that Hlsl OPs are getting executed
* Win64 Editor PCD3D_SM5: manually checked that shaders don't get compiled and that model creation fails with unsupported messages.
* Win64 Editor VULKAN_SM5: manually checked that shaders don't get compiled and that model creation fails with unsupported messages.
* Win64 Editor VULKAN_SM6: manually checked that shaders don't get compiled and that model creation fails with unsupported messages.
* Win64 Standalone PCD3D_SM6: manually checked that Hlsl OPs are getting executed
* Win64 Standalone PCD3D_SM5: manually checked that shaders don't get compiled, that Hlsl runtime doesn't get registered and that proper messages are emitted
* Win64 Standalone VULKAN_SM5: manually checked that shaders don't get compiled, that Hlsl runtime doesn't get registered and that proper messages are emitted
* Win64 Standalone VULKAN_SM6:manually checked that shaders don't get compiled, that Hlsl runtime doesn't get registered and that proper messages are emitted
* Supported Consoles: manually checked that Hlsl OPs are getting executed
* Mac Editor METAL_SM5: manually checked that shaders don't get compiled and that model creation fails with unsupported messages.
* Mac Editor METAL_SM6: manually checked that shaders are getting compiled (although compilation fails for an unrelated reason)
* Linux Editor VULKAN_SM5: manually checked that shaders don't get compiled and that model creation fails with unsupported messages.
* Linux Editor VULKAN_SM6: manually checked that shaders don't get compiled and that model creation fails with unsupported messages.
* Linux Standalone VULKAN_SM6:manually checked that shaders don't get compiled, that Hlsl runtime doesn't get registered and that proper messages are emitted

#rnx

diff --git a/Engine/Config/Mac/DataDrivenPlatformInfo.ini b/Engine/Config/Mac/DataDrivenPlatformInfo.ini
--- a/Engine/Config/Mac/DataDrivenPlatformInfo.ini
+++ b/Engine/Config/Mac/DataDrivenPlatformInfo.ini
@@ -60,7 +60,7 @@ bSupportsByteBufferComputeShaders = true
 bSupportsHairStrandGeometry = true
 bSupportsNanite=false
 bSupportsUInt64ImageAtomics=false
-bSupportsNNEShaders=true
+bSupportsNNEShaders=false
 bSupportsIndependentSamplers=true
 bSupportsComputeFramework = true
 FriendlyName=LOCTEXT("FriendlyShaderPlatformName_MacSM5", "Mac Metal SM5")
diff --git a/Engine/Config/VulkanPC/DataDrivenPlatformInfo.ini b/Engine/Config/VulkanPC/DataDrivenPlatformInfo.ini
--- a/Engine/Config/VulkanPC/DataDrivenPlatformInfo.ini
+++ b/Engine/Config/VulkanPC/DataDrivenPlatformInfo.ini
@@ -133,7 +133,7 @@ EnablesHLSL2021ByDefault=1
 FriendlyName=LOCTEXT("FriendlyShaderPlatformName_VulkanPC_SM5", "PC Vulkan SM5")
 
 BindlessSupport=Unsupported
-bSupportsNNEShaders=true
+bSupportsNNEShaders=false
 bSupportsUniformBufferObjects = true
 bSupportsRealTypes=Unsupported
 bSupportsIndependentSamplers=true
@@ -258,7 +258,7 @@ EnablesHLSL2021ByDefault=1
 FriendlyName=LOCTEXT("FriendlyShaderPlatformName_VulkanPC_SM6", "PC Vulkan SM6")
 
 BindlessSupport=AllShaderTypes
-bSupportsNNEShaders=true
+bSupportsNNEShaders=false
 bSupportsUniformBufferObjects = true
 bSupportsRealTypes=Unsupported
 bSupportsIndependentSamplers=true
diff --git a/Engine/Config/Windows/DataDrivenPlatformInfo.ini b/Engine/Config/Windows/DataDrivenPlatformInfo.ini
--- a/Engine/Config/Windows/DataDrivenPlatformInfo.ini
+++ b/Engine/Config/Windows/DataDrivenPlatformInfo.ini
@@ -74,7 +74,7 @@ StaticShaderBindingLayoutSupport=RayTracingOnly
 SupportsMultiViewport=RuntimeDependent
 FriendlyName=LOCTEXT("FriendlyShaderPlatformName_D3D_SM5", "PC D3D SM5")
 bSupportsDualSourceBlending=true
-bSupportsNNEShaders=true
+bSupportsNNEShaders=false
 bSupportsIndependentSamplers=true
 bSupportsVertexShaderUAVs=RuntimeGuaranteed
 
diff --git a/Engine/Plugins/Experimental/NNERuntimeRDG/NNERuntimeRDG.uplugin b/Engine/Plugins/Experimental/NNERuntimeRDG/NNERuntimeRDG.uplugin
--- a/Engine/Plugins/Experimental/NNERuntimeRDG/NNERuntimeRDG.uplugin
+++ b/Engine/Plugins/Experimental/NNERuntimeRDG/NNERuntimeRDG.uplugin
@@ -15,21 +15,16 @@
 	"IsBetaVersion": false,
 	"IsExperimentalVersion": true,
 	"Installed": false,
-	"SupportedTargetPlatforms": [
-		"Win64", "Linux", "Mac"
-	],
 	"Modules": [
 		{
 			"Name": "NNEHlslShaders",
 			"Type": "RuntimeAndProgram",
-			"LoadingPhase": "PostConfigInit",
-			"PlatformAllowList": [ "Win64", "Linux", "Mac" ]
+			"LoadingPhase": "PostConfigInit"
 		},
 		{
 			"Name": "NNERuntimeRDG",
 			"Type": "RuntimeAndProgram",
-			"LoadingPhase": "Default",
-			"PlatformAllowList": [ "Win64", "Linux", "Mac" ]
+			"LoadingPhase": "Default"
 		},
 		{
 			"Name": "NNERuntimeRDGUtils",
diff --git a/Engine/Plugins/Experimental/NNERuntimeRDG/Source/NNEHlslShaders/Internal/NNEHlslShadersCumSumCS.h b/Engine/Plugins/Experimental/NNERuntimeRDG/Source/NNEHlslShaders/Internal/NNEHlslShadersCumSumCS.h
--- a/Engine/Plugins/Experimental/NNERuntimeRDG/Source/NNEHlslShaders/Internal/NNEHlslShadersCumSumCS.h
+++ b/Engine/Plugins/Experimental/NNERuntimeRDG/Source/NNEHlslShaders/Internal/NNEHlslShadersCumSumCS.h
@@ -76,6 +76,5 @@ namespace UE::NNEHlslShaders::Internal
 		END_SHADER_PARAMETER_STRUCT()
 
 		static void ModifyCompilationEnvironment(const FGlobalShaderPermutationParameters& InParameters, FShaderCompilerEnvironment& OutEnvironment);
-		static bool ShouldCompilePermutation(const FGlobalShaderPermutationParameters& InParameters);
 	};
 } // UE::NNEHlslShaders::Internal
\ No newline at end of file
diff --git a/Engine/Plugins/Experimental/NNERuntimeRDG/Source/NNEHlslShaders/Private/NNEHlslShadersCumSumCS.cpp b/Engine/Plugins/Experimental/NNERuntimeRDG/Source/NNEHlslShaders/Private/NNEHlslShadersCumSumCS.cpp
--- a/Engine/Plugins/Experimental/NNERuntimeRDG/Source/NNEHlslShaders/Private/NNEHlslShadersCumSumCS.cpp
+++ b/Engine/Plugins/Experimental/NNERuntimeRDG/Source/NNEHlslShaders/Private/NNEHlslShadersCumSumCS.cpp
@@ -24,16 +24,6 @@ namespace UE::NNEHlslShaders::Internal
 		FPermutationDomain PermutationVector(InParameters.PermutationId);
 	}
 
-	bool TCumSumCS::ShouldCompilePermutation(const FGlobalShaderPermutationParameters& InParameters)
-	{
-		//NOTE: platform compatibility checks should be rather done in a NNERuntimeRDG-wise place and called from here
-
-		const ERHIFeatureSupport WaveOpsSupport = FDataDrivenShaderPlatformInfo::GetSupportsWaveOperations(InParameters.Platform);
-		const bool bWaveOpsAvailable = WaveOpsSupport == ERHIFeatureSupport::RuntimeDependent || WaveOpsSupport == ERHIFeatureSupport::RuntimeGuaranteed;
-
-		return bWaveOpsAvailable;
-	}
-
 	IMPLEMENT_GLOBAL_SHADER(TInitCumSumCS, "/NNEHlslShaders/NNEHlslShadersCumSum.usf", "InitCumSum", SF_Compute);
 	IMPLEMENT_GLOBAL_SHADER(TCumSumCS, "/NNEHlslShaders/NNEHlslShadersCumSum.usf", "CumSum", SF_Compute);
 } // UE::NNEHlslShaders::Internal
diff --git a/Engine/Plugins/Experimental/NNERuntimeRDG/Source/NNEHlslShaders/Private/NNEHlslShadersSplitCS.cpp b/Engine/Plugins/Experimental/NNERuntimeRDG/Source/NNEHlslShaders/Private/NNEHlslShadersSplitCS.cpp
--- a/Engine/Plugins/Experimental/NNERuntimeRDG/Source/NNEHlslShaders/Private/NNEHlslShadersSplitCS.cpp
+++ b/Engine/Plugins/Experimental/NNERuntimeRDG/Source/NNEHlslShaders/Private/NNEHlslShadersSplitCS.cpp
@@ -7,6 +7,11 @@ namespace UE::NNEHlslShaders::Internal
 {
 	bool FSplitCS::ShouldCompilePermutation(const FGlobalShaderPermutationParameters& InParameters)
 	{
+		if (!FHlslShaderBase::ShouldCompilePermutation(InParameters))
+		{
+			return false;
+		}
+
 		const FPermutationDomain PermutationVector(InParameters.PermutationId);
 
 		return PermutationVector.Get<FSplitCS::FSplitAxis>() < PermutationVector.Get<FSplitCS::FSplitRank>();
diff --git a/Engine/Plugins/Experimental/NNERuntimeRDG/Source/NNERuntimeRDG/NNERuntimeRDG.Build.cs b/Engine/Plugins/Experimental/NNERuntimeRDG/Source/NNERuntimeRDG/NNERuntimeRDG.Build.cs
--- a/Engine/Plugins/Experimental/NNERuntimeRDG/Source/NNERuntimeRDG/NNERuntimeRDG.Build.cs
+++ b/Engine/Plugins/Experimental/NNERuntimeRDG/Source/NNERuntimeRDG/NNERuntimeRDG.Build.cs
@@ -38,12 +38,31 @@ public class NNERuntimeRDG : ModuleRules
 			PrivateDependencyModuleNames.Add("VulkanRHI");
 		}
 
-		if ((Target.Type == TargetType.Editor || Target.Type == TargetType.Program) &&
-			(Target.Platform == UnrealTargetPlatform.Win64 || Target.Platform == UnrealTargetPlatform.Linux || Target.Platform == UnrealTargetPlatform.Mac)
+		if (Target.Type == TargetType.Editor || Target.Type == TargetType.Program)
+		{
+			// Supported platforms (editor)
+			if ( Target.Platform == UnrealTargetPlatform.Win64 || 
+				Target.Platform == UnrealTargetPlatform.Linux || 
+				Target.Platform == UnrealTargetPlatform.Mac
 			)
+			{
+				PublicDefinitions.Add("WITH_NNE_RUNTIME_HLSL");
+				PrivateDefinitions.Add("NNE_UTILITIES_AVAILABLE");
+				PrivateDependencyModuleNames.Add("NNERuntimeRDGUtils");
+			}
+		}
+		else
 		{
-			PrivateDefinitions.Add("NNE_UTILITIES_AVAILABLE");
-			PrivateDependencyModuleNames.Add("NNERuntimeRDGUtils");
+			// Supported platforms (standalone)
+			// NOTE: To add a supported platform:
+			// - either here or in a NNERuntimeRDG_*.Build.cs extension, set WITH_NNE_RUNTIME_HLSL for Game/Client target type,
+			// - set bSupportsNNEShaders=true in DataDrivenPlatformInfo.ini for the platform - shader format pair.
+			if ( Target.Platform == UnrealTargetPlatform.Win64 || 
+				 Target.Platform == UnrealTargetPlatform.Mac
+			   )
+			{
+				PublicDefinitions.Add("WITH_NNE_RUNTIME_HLSL");
+			}
 		}
 
 		PublicDefinitions.Add("NNERUNTIMERDGHLSL_BUFFER_LENGTH_ALIGNMENT=4");
diff --git a/Engine/Plugins/Experimental/NNERuntimeRDG/Source/NNERuntimeRDG/Private/Hlsl/NNERuntimeRDGCumSum.cpp b/Engine/Plugins/Experimental/NNERuntimeRDG/Source/NNERuntimeRDG/Private/Hlsl/NNERuntimeRDGCumSum.cpp
--- a/Engine/Plugins/Experimental/NNERuntimeRDG/Source/NNERuntimeRDG/Private/Hlsl/NNERuntimeRDGCumSum.cpp
+++ b/Engine/Plugins/Experimental/NNERuntimeRDG/Source/NNERuntimeRDG/Private/Hlsl/NNERuntimeRDGCumSum.cpp
@@ -130,13 +130,6 @@ namespace UE::NNERuntimeRDG::Private::Hlsl
 			check(InputTensorDescs.Num() == 2);
 			check(OutputTensorDescs.Num() == 1);
 
-			//NOTE: hardware compatibility check should rather be done in a NNERuntimeRDG-wise place
-			if(!GRHISupportsWaveOperations)
-			{
-				UE_LOG(LogNNE, Warning, TEXT("Hlsl CumSum not supported on current hardware."));
-				return false;
-			}
-
 			const int32 InputRank = InputTensorDescs[0].GetShape().Rank();
 			
 			int32 Exclusive = Attributes.GetValueOrDefault<int32>(TEXT("exclusive"), 0);
diff --git a/Engine/Plugins/Experimental/NNERuntimeRDG/Source/NNERuntimeRDG/Private/NNERuntimeRDGHlsl.cpp b/Engine/Plugins/Experimental/NNERuntimeRDG/Source/NNERuntimeRDG/Private/NNERuntimeRDGHlsl.cpp
--- a/Engine/Plugins/Experimental/NNERuntimeRDG/Source/NNERuntimeRDG/Private/NNERuntimeRDGHlsl.cpp
+++ b/Engine/Plugins/Experimental/NNERuntimeRDG/Source/NNERuntimeRDG/Private/NNERuntimeRDGHlsl.cpp
@@ -117,6 +117,33 @@ namespace ConsoleCommands
 	);
 } // ConsoleCommands
 
+bool UNNERuntimeRDGHlslImpl::IsCurrentPlatformSupported()
+{
+	bool bResult = true;
+
+#ifndef NNE_FORCE_HARDWARE_SUPPORTS_HLSL
+	if(GMaxRHIFeatureLevel < ERHIFeatureLevel::SM5)
+	{
+		UE_LOG(LogNNE, Display, TEXT("UNNERuntimeRDGHlsl: minimum feature level required is SM5 for current RHI platform."));
+		bResult &= false;
+	}
+
+	if(!GRHISupportsWaveOperations)
+	{
+		UE_LOG(LogNNE, Display, TEXT("UNNERuntimeRDGHlsl: current RHI platform doesn't support wave operations."));
+		bResult &= false;
+	}
+
+	if(!GRHIGlobals.SupportsNative16BitOps)
+	{
+		UE_LOG(LogNNE, Display, TEXT("UNNERuntimeRDGHlsl: current RHI platform doesn't support native 16-bit operations."));
+		bResult &= false;
+	}
+#endif
+
+	return bResult;
+}
+
 UNNERuntimeRDGHlslImpl::ECanCreateModelDataStatus UNNERuntimeRDGHlslImpl::CanCreateModelData(const FString& FileType, TConstArrayView<uint8> FileData, const TMap<FString, TConstArrayView<uint8>>& AdditionalFileData, const FGuid& FileId, const ITargetPlatform* TargetPlatform) const
 {
 #ifdef NNE_UTILITIES_AVAILABLE
@@ -147,6 +174,8 @@ UNNERuntimeRDGHlslImpl::ECanCreateModelRDGStatus UNNERuntimeRDGHlslImpl::CanCrea
 	bool bResult = FGenericPlatformMemory::Memcmp(&(Data[0]), &(GUID), GuidSize) == 0;
 	bResult &= FGenericPlatformMemory::Memcmp(&(Data[GuidSize]), &(Version), VersionSize) == 0;
 
+	bResult &= IsCurrentPlatformSupported();
+
 	return bResult ? ECanCreateModelRDGStatus::Ok : ECanCreateModelRDGStatus::Fail;
 };
 
diff --git a/Engine/Plugins/Experimental/NNERuntimeRDG/Source/NNERuntimeRDG/Private/NNERuntimeRDGHlsl.h b/Engine/Plugins/Experimental/NNERuntimeRDG/Source/NNERuntimeRDG/Private/NNERuntimeRDGHlsl.h
--- a/Engine/Plugins/Experimental/NNERuntimeRDG/Source/NNERuntimeRDG/Private/NNERuntimeRDGHlsl.h
+++ b/Engine/Plugins/Experimental/NNERuntimeRDG/Source/NNERuntimeRDG/Private/NNERuntimeRDGHlsl.h
@@ -28,6 +28,8 @@ public:
 
 	virtual FString GetRuntimeName() const override { return TEXT("NNERuntimeRDGHlsl"); };
 
+	static bool IsCurrentPlatformSupported();
+
 	virtual ECanCreateModelDataStatus CanCreateModelData(const FString& FileType, TConstArrayView<uint8> FileData, const TMap<FString, TConstArrayView<uint8>>& AdditionalFileData, const FGuid& FileId, const ITargetPlatform* TargetPlatform) const override;
 	virtual TSharedPtr<UE::NNE::FSharedModelData> CreateModelData(const FString& FileType, TConstArrayView<uint8> FileData, const TMap<FString, TConstArrayView<uint8>>& AdditionalFileData, const FGuid& FileId, const ITargetPlatform* TargetPlatform) override;
 	virtual FString GetModelDataIdentifier(const FString& FileType, TConstArrayView<uint8> FileData, const TMap<FString, TConstArrayView<uint8>>& AdditionalFileData, const FGuid& FileId, const ITargetPlatform* TargetPlatform) const override;
diff --git a/Engine/Plugins/Experimental/NNERuntimeRDG/Source/NNERuntimeRDG/Private/NNERuntimeRDGModule.cpp b/Engine/Plugins/Experimental/NNERuntimeRDG/Source/NNERuntimeRDG/Private/NNERuntimeRDGModule.cpp
--- a/Engine/Plugins/Experimental/NNERuntimeRDG/Source/NNERuntimeRDG/Private/NNERuntimeRDGModule.cpp
+++ b/Engine/Plugins/Experimental/NNERuntimeRDG/Source/NNERuntimeRDG/Private/NNERuntimeRDGModule.cpp
@@ -5,6 +5,7 @@
 #include "NNE.h"
 #include "NNERuntimeRDGHlsl.h"
 #include "UObject/WeakInterfacePtr.h"
+#include "DataDrivenShaderPlatformInfo.h"
 
 static TAutoConsoleVariable<int32> CVarHlslModelOptimization(
 	TEXT("nne.hlsl.ModelOptimization"),
@@ -17,31 +18,61 @@ static TAutoConsoleVariable<int32> CVarHlslModelOptimization(
 #endif
 	);
 
-void FNNERuntimeRDGModule::StartupModule()
+void FNNERuntimeRDGModule::RegisterRuntime()
 {
-	// NNE runtime ORT Cpu startup
 	NNERuntimeRDGHlsl = NewObject<UNNERuntimeRDGHlslImpl>();
 	if (NNERuntimeRDGHlsl.IsValid())
 	{
-		TWeakInterfacePtr<INNERuntime> RuntimeCPUInterface(NNERuntimeRDGHlsl.Get());
+		TWeakInterfacePtr<INNERuntime> RuntimeInterface(NNERuntimeRDGHlsl.Get());
 
 		NNERuntimeRDGHlsl->Init();
 		NNERuntimeRDGHlsl->AddToRoot();
-		UE::NNE::RegisterRuntime(RuntimeCPUInterface);
+		UE::NNE::RegisterRuntime(RuntimeInterface);
 	}
 }
 
+void FNNERuntimeRDGModule::StartupModule()
+{
+#if WITH_EDITOR
+	#ifdef WITH_NNE_RUNTIME_HLSL 
+		RegisterRuntime();
+	#endif
+#else
+	#ifdef WITH_NNE_RUNTIME_HLSL 
+		if(FDataDrivenShaderPlatformInfo::GetSupportsNNEShaders(GMaxRHIShaderPlatform))
+		{
+			if(UNNERuntimeRDGHlslImpl::IsCurrentPlatformSupported())
+			{
+				RegisterRuntime();
+			}
+			else
+			{
+				UE_LOG(LogNNE, Display, TEXT("FNNERuntimeRDGModule: not registering NNERuntimeRDGHlsl runtime because current hardware is incompatible, consider bypassing by setting the define NNE_FORCE_HARDWARE_SUPPORTS_HLSL."));
+			}
+		}
+		else
+		{
+			UE_LOG(LogNNE, Display, TEXT("FNNERuntimeRDGModule: not registering NNERuntimeRDGHlsl runtime because current RHI shader platform is not enabled, consider setting the flag bSupportsNNEShaders in DataDrivenPlatformInfo."));
+		}
+	#else
+		UE_LOG(LogNNE, Display, TEXT("FNNERuntimeRDGModule: not registering NNERuntimeRDGHlsl runtime as platform is not enabled, if needed set define WITH_NNE_RUNTIME_HLSL."));
+	#endif
+#endif
+}
+
 void FNNERuntimeRDGModule::ShutdownModule()
 {
-	// NNE runtime ORT Cpu shutdown
+
+	// NNE runtime RDG Hlsl shutdown
 	if (NNERuntimeRDGHlsl.IsValid())
 	{
-		TWeakInterfacePtr<INNERuntime> RuntimeCPUInterface(NNERuntimeRDGHlsl.Get());
+		TWeakInterfacePtr<INNERuntime> RuntimeInterface(NNERuntimeRDGHlsl.Get());
 
-		UE::NNE::UnregisterRuntime(RuntimeCPUInterface);
+		UE::NNE::UnregisterRuntime(RuntimeInterface);
 		NNERuntimeRDGHlsl->RemoveFromRoot();
 		NNERuntimeRDGHlsl.Reset();
 	}
+
 }
 
 IMPLEMENT_MODULE(FNNERuntimeRDGModule, NNERuntimeRDG);
diff --git a/Engine/Plugins/Experimental/NNERuntimeRDG/Source/NNERuntimeRDG/Private/NNERuntimeRDGModule.h b/Engine/Plugins/Experimental/NNERuntimeRDG/Source/NNERuntimeRDG/Private/NNERuntimeRDGModule.h
--- a/Engine/Plugins/Experimental/NNERuntimeRDG/Source/NNERuntimeRDG/Private/NNERuntimeRDGModule.h
+++ b/Engine/Plugins/Experimental/NNERuntimeRDG/Source/NNERuntimeRDG/Private/NNERuntimeRDGModule.h
@@ -12,6 +12,8 @@ class UNNERuntimeRDGHlslImpl;
 
 class FNNERuntimeRDGModule : public IModuleInterface
 {
+private:
+	void RegisterRuntime();
 
 public:
 	TWeakObjectPtr<UNNERuntimeRDGHlslImpl> NNERuntimeRDGHlsl{ nullptr };
