NNE - Fix &= in NNERuntimeRDGHlsl.cpp

Problems:
* Usage of &= false in NNERUntimeRDGHlsl.cpp gives static analysis warnings

Solutions:
* Substitute with = false.

Tests:
* Win64 Editor PCD3D_SM6: manually checked that everything works
* Win64 Editor PCD3D_SM5: manually checked that model creation fails

#rnx

diff --git a/Engine/Plugins/Experimental/NNERuntimeRDG/Source/NNERuntimeRDG/Private/NNERuntimeRDGHlsl.cpp b/Engine/Plugins/Experimental/NNERuntimeRDG/Source/NNERuntimeRDG/Private/NNERuntimeRDGHlsl.cpp
--- a/Engine/Plugins/Experimental/NNERuntimeRDG/Source/NNERuntimeRDG/Private/NNERuntimeRDGHlsl.cpp
+++ b/Engine/Plugins/Experimental/NNERuntimeRDG/Source/NNERuntimeRDG/Private/NNERuntimeRDGHlsl.cpp
@@ -125,19 +125,19 @@ bool UNNERuntimeRDGHlslImpl::IsCurrentPlatformSupported()
 	if(GMaxRHIFeatureLevel < ERHIFeatureLevel::SM5)
 	{
 		UE_LOG(LogNNE, Display, TEXT("UNNERuntimeRDGHlsl: minimum feature level required is SM5 for current RHI platform."));
-		bResult &= false;
+		bResult = false;
 	}
 
 	if(!GRHISupportsWaveOperations)
 	{
 		UE_LOG(LogNNE, Display, TEXT("UNNERuntimeRDGHlsl: current RHI platform doesn't support wave operations."));
-		bResult &= false;
+		bResult = false;
 	}
 
 	if(!GRHIGlobals.SupportsNative16BitOps)
 	{
 		UE_LOG(LogNNE, Display, TEXT("UNNERuntimeRDGHlsl: current RHI platform doesn't support native 16-bit operations."));
-		bResult &= false;
+		bResult = false;
 	}
 #endif
 
