[NNERuntimeRDG HLSL] Remove some of the MatMul Convolution restrictions

Problem: The MatMul implementation of the Convolution operator has restrictions on the shape that it can handle. For the OIDN operators particularly relevant are the restrictions that the input channel count (Ci) needs to be a multiple of 16 and that the output channel count (Cw) needs to be a multiple of 32. This means that for the most expensive convolution nodes in an OIDN model we have to fallback to the default Convolution implementation, which is significantly slower.

Solution: Perform the same number of operations as if Ci and Cw would be a multiple of 16 and 32 respectively, but set loaded values to zero and skip storing values if the channel is out of bound.

Tested:
* Win64 Standalone tests: all
* Win64 Editor tests: all
* Consoles tests: convolution, OIDN and smoke tests

#rnx

#virtualized

diff --git a/Engine/Plugins/Experimental/NNERuntimeRDG/Shaders/Private/NNEHlslShaders/NNEHlslShadersConvMatmul.usf b/Engine/Plugins/Experimental/NNERuntimeRDG/Shaders/Private/NNEHlslShaders/NNEHlslShadersConvMatmul.usf
--- a/Engine/Plugins/Experimental/NNERuntimeRDG/Shaders/Private/NNEHlslShaders/NNEHlslShadersConvMatmul.usf
+++ b/Engine/Plugins/Experimental/NNERuntimeRDG/Shaders/Private/NNEHlslShaders/NNEHlslShadersConvMatmul.usf
@@ -110,6 +110,9 @@ void ConvMatmul(
 	const int InputChannelOffset = GroupThreadIdx >= 32 ? LDS_IC/2 : 0;
 	const int OutputPixelOffset = GroupThreadIdx % 32;
 	const int OutputChannelOffset = GroupThreadIdx % 32;
+	int OutputChannel = Scalar_GroupOutputChannelOffset + OutputChannelOffset;
+	const bool IsValidOutputChannel = OutputChannel < Cw;
+	OutputChannel = min(OutputChannel, Cw - 1);
 	
 	for (int hw = 0; hw < Hw; ++hw)
 	{
@@ -119,10 +122,19 @@ void ConvMatmul(
 		for (int ww = 0; ww < Ww; ++ww)
 		{
 			const int InputPixelW = Scalar_GroupInputPixelKernelTopLeftW + ww + OutputPixelOffset * StrideW;
-			const bool IsValidInput = Scalar_IsValidInputH && (InputPixelW >= 0) && (InputPixelW < Wi);
+			const bool IsValidInputPixel = Scalar_IsValidInputH && (InputPixelW >= 0) && (InputPixelW < Wi);
 			int InputPixelOffset = Scalar_InputPixelH * Wi + InputPixelW;
+			InputPixelOffset = clamp(InputPixelOffset, 0, Hi * Wi - 1);
 			
-			InputPixelOffset = clamp(InputPixelOffset, 0, Hi * Wi);
+			const int ReadOffsetI = Scalar_BatchInputOffset + InputPixelOffset;
+
+			#if WEIGHTS_TRANSPOSED
+			const int Scalar_InputChannelStride = Cw;
+			const int ReadOffsetW = Scalar_KernelIdx + OutputChannel;
+			#else
+			const int Scalar_InputChannelStride = Hw * Ww;
+			const int ReadOffsetW = OutputChannel * Ci * Scalar_InputChannelStride + Scalar_KernelIdx;
+			#endif
 			
 			for (int ci = 0; ci < Ci; ci += LDS_IC)
 			{
@@ -132,23 +144,20 @@ void ConvMatmul(
 				// 512 inputs (16 inputs channel, 32 inputs pixels)
 				// We have 64 threads thus each will read 8 inputs and 8 weights.
 
-				const int ReadOffsetI = Scalar_BatchInputOffset + (ci + InputChannelOffset) * Hi * Wi + InputPixelOffset;
-				#if WEIGHTS_TRANSPOSED
-					const int InputChannelStride = Cw;
-					const int ReadOffsetW = Scalar_KernelIdx + (ci + InputChannelOffset) * InputChannelStride + (Scalar_GroupOutputChannelOffset + OutputChannelOffset);
-				#else
-					const int InputChannelStride = Hw * Ww;
-					const int ReadOffsetW = (Scalar_GroupOutputChannelOffset + OutputChannelOffset) * Ci * InputChannelStride + (ci + InputChannelOffset) * InputChannelStride + Scalar_KernelIdx;
-				#endif
 				
 				UNROLL
 				for (i = 0; i < (LDS_IC / 2); ++i)
 				{
-					WORK_TYPE ValueI = BUFFER_TO_WORK_TYPE(Input[Hi * Wi * i + ReadOffsetI]);
-					ValueI = IsValidInput ? ValueI : 0.0f;
+					int InputChannel = ci + InputChannelOffset + i;
+					const bool IsValidInputChannel = InputChannel < Ci;
+					InputChannel = min(InputChannel, Ci - 1);
+
+					WORK_TYPE ValueI = BUFFER_TO_WORK_TYPE(Input[Hi * Wi * InputChannel + ReadOffsetI]);
+					ValueI = IsValidInputChannel && IsValidInputPixel ? ValueI : 0.0f;
 					WriteLDSInputs(InputChannelOffset + i, OutputPixelOffset, ValueI);
 
-					WORK_TYPE ValueW = BUFFER_TO_WORK_TYPE(Weight[InputChannelStride * i + ReadOffsetW]);
+					WORK_TYPE ValueW = BUFFER_TO_WORK_TYPE(Weight[Scalar_InputChannelStride * InputChannel + ReadOffsetW]);
+					ValueW = IsValidInputChannel && IsValidOutputChannel ? ValueW : 0.0f;
 					WriteLDSWeights(InputChannelOffset + i, OutputChannelOffset, ValueW);
 				}
 				
@@ -216,8 +225,11 @@ void ConvMatmul(
 	UNROLL
 	for (i = 0; i < LDS_IC; ++i)
 	{
-		const int WriteOffset = Scalar_BatchOutputOffset + (Scalar_GroupOutputChannelOffset + i + WriteOutChannelOffset) * Ho * Wo + Scalar_GroupOutputPixelOffset + WriteOutPixelOffset;
-		Output[WriteOffset] = LDSBuffer[(i + WriteOutChannelOffset) * LDS_X_PADDEDSIZE + LDSPixelOffset];
+		const int WriteOutChannel = Scalar_GroupOutputChannelOffset + i + WriteOutChannelOffset;
+		if (WriteOutChannel < Cw)
+		{
+			const int WriteOffset = Scalar_BatchOutputOffset + WriteOutChannel * Ho * Wo + Scalar_GroupOutputPixelOffset + WriteOutPixelOffset;
+			Output[WriteOffset] = LDSBuffer[(i + WriteOutChannelOffset) * LDS_X_PADDEDSIZE + LDSPixelOffset];
+		}
 	}
-
-}
\ No newline at end of file
+}
diff --git a/Engine/Plugins/Experimental/NNERuntimeRDG/Source/NNERuntimeRDG/Private/Hlsl/NNERuntimeRDGConv.cpp b/Engine/Plugins/Experimental/NNERuntimeRDG/Source/NNERuntimeRDG/Private/Hlsl/NNERuntimeRDGConv.cpp
--- a/Engine/Plugins/Experimental/NNERuntimeRDG/Source/NNERuntimeRDG/Private/Hlsl/NNERuntimeRDGConv.cpp
+++ b/Engine/Plugins/Experimental/NNERuntimeRDG/Source/NNERuntimeRDG/Private/Hlsl/NNERuntimeRDGConv.cpp
@@ -41,6 +41,7 @@ namespace UE::NNERuntimeRDG::Private::Hlsl
 
 		EConvAutoPad AutoPad = EConvAutoPad::NOTSET;
 		TArray<int32> Dilations;
+		bool bHasDilation;
 		int32 Group = 1;
 		TArray<int32> Pads;
 		TArray<int32> Strides;
@@ -113,6 +114,7 @@ namespace UE::NNERuntimeRDG::Private::Hlsl
 
 			FConvCS::LexFromString(AutoPad, *Attributes.GetValueOrDefault<FString>(TEXT("auto_pad"), TEXT("NOTSET")));
 			Dilations = Attributes.GetValueOrDefault<TArray<int32>>(TEXT("dilations"), DilationsOrStridesDefault);
+			bHasDilation = Algo::AnyOf(Dilations, [](auto Dim) {return Dim != 1u; });
 			Group = Attributes.GetValueOrDefault<int32>(TEXT("group"), 1);
 			if (AutoPad == EConvAutoPad::NOTSET)
 			{
@@ -180,10 +182,6 @@ namespace UE::NNERuntimeRDG::Private::Hlsl
 				return false;
 			if (Weights.GetShape().Rank() != 4)
 				return false;
-			
-			bool bHasBias = Bias != nullptr;
-
-			const bool bHasDilation = Algo::AnyOf(Dilations, [](auto Dim) {return Dim != 1u; });
 			if (bHasDilation)
 				return false;
 
@@ -202,14 +200,11 @@ namespace UE::NNERuntimeRDG::Private::Hlsl
 			int Hw = Weights.GetShape().GetData()[2];
 			int Ww = Weights.GetShape().GetData()[3];
 
-			if (Ci % 16 != 0)
-				return false;
-			if (Cw % 32 != 0)
-				return false;
 			if (Wo % 32 != 0) // Idea : support this by launching more threads and discard some results so a threadgroup still operator on only one value for H.
 				return false;
 
 			TArray<int32> Padding = FConvCS::GetPadding(Input.GetShape().GetData(), Weights.GetShape().GetData(), AutoPad, Dilations, Strides, Pads);
+			bool bHasBias = Bias != nullptr;
 
 			// Set parameters
 			FConvMatmulCS::FParameters* Params = GraphBuilder.AllocParameters<FConvMatmulCS::FParameters>();
@@ -293,16 +288,10 @@ namespace UE::NNERuntimeRDG::Private::Hlsl
 			// Heuristics : only matmul implementation benefits from transposed weights
 			if (Weights->GetShape().Rank() != 4)
 				return;
-
-			uint32 Cw = Weights->GetShape().GetData()[0];
-			uint32 Ci = Weights->GetShape().GetData()[1];
-
-			if (Ci % 16 != 0)
-				return;
-			if (Cw % 32 != 0)
-				return;
 			if (Group != 1)
 				return;
+			if (bHasDilation)
+				return;
 
 			//Transpose from CwCiHwWw to HwWwCiCw
 			if (Internal::CPUHelper::Transpose::TransposePreparedData(*Weights, {2,3,1,0} ))
