[NNERuntimeRDG HLSL] Add Constant operator

Problem: Unoptimized OIDN models use the constant operator which is currently not implemented in HLSL. This makes it impossible to test fp16 support.

Solution: Implement constant operator.

Tested:
* Win64 Standalone tests: all
* Win64 Editor tests: all
* Console tests: smoke, constant, intializer

#rnx

#virtualized

diff --git a/Engine/Plugins/Experimental/NNERuntimeRDG/Source/NNERuntimeRDG/Private/Hlsl/NNERuntimeRDGConstant.cpp b/Engine/Plugins/Experimental/NNERuntimeRDG/Source/NNERuntimeRDG/Private/Hlsl/NNERuntimeRDGConstant.cpp
new file mode 100644
--- /dev/null
+++ b/Engine/Plugins/Experimental/NNERuntimeRDG/Source/NNERuntimeRDG/Private/Hlsl/NNERuntimeRDGConstant.cpp
@@ -0,0 +1,156 @@
+// Copyright Epic Games, Inc. All Rights Reserved.
+
+#include "NNERuntimeRDGConstant.h"
+
+#include "Helper/NNERuntimeRDGLogHelper.h"
+#include "Helper/NNERuntimeRDGOperatorHelper.h"
+#include "NNEAttributeTensor.h"
+#include "NNERuntimeRDGHlslHelper.h"
+#include "NNETensor.h"
+#include "NNETypes.h"
+
+namespace UE::NNERuntimeRDG::Private::Hlsl
+{
+	/**
+	 * Constant operator implementation
+	 */
+	class FConstant : public FOperatorHlsl
+	{
+	public:
+
+		FConstant() {}
+		virtual ~FConstant() = default;
+
+	private:
+
+		FNNEAttributeValue Attribute;
+
+
+	public:
+
+		virtual int PrepareOutputs(TConstArrayView<NNE::Internal::FTensorRef> InputTensors, TArrayView<NNE::Internal::FTensorRef> OutputTensors) override
+		{
+			using NNE::FTensorShape;
+			using NNE::Internal::FTensor;
+			using NNE::Internal::FAttributeTensor;
+
+			check(InputTensors.Num() == 0);
+			check(OutputTensors.Num() == 1);
+
+			FTensor& Output = *OutputTensors[0];
+
+			switch (Attribute.GetType())
+			{
+				case ENNEAttributeDataType::Float:
+				{
+					if (Output.GetDataType() != ENNETensorDataType::Float)
+					{
+						UE_LOG(LogNNE, Warning, TEXT("Output data type %s does not match constant type of float"), 
+							   					*LogHelper::GetTensorDataTypeName(Output.GetDataType()));
+						return -1;
+					}
+					const float Value = Attribute.GetValue<float>();
+					const TArray<uint32, TInlineAllocator<FTensorShape::MaxRank>> Shape;
+					const TConstArrayView<float> Data = MakeArrayView(&Value, 1);
+					Output.SetShape(FTensorShape::Make(Shape));
+					Output.SetPreparedData(Data);
+				} break;
+				case ENNEAttributeDataType::FloatArray:
+				{
+					if (Output.GetDataType() != ENNETensorDataType::Float)
+					{
+						UE_LOG(LogNNE, Warning, TEXT("Output data type %s does not match constant type of float"), 
+							   					*LogHelper::GetTensorDataTypeName(Output.GetDataType()));
+						return -1;
+					}
+					const TArray<float> Values = Attribute.GetValue<TArray<float>>();
+					TArray<uint32, TInlineAllocator<FTensorShape::MaxRank>> Shape;
+					Shape.Add(Values.Num());
+					Output.SetShape(FTensorShape::Make(Shape));
+					Output.SetPreparedData(MakeArrayView(Values));
+				} break;
+				case ENNEAttributeDataType::Tensor:
+				{
+					const FAttributeTensor AttributeTensor = Attribute.GetValue<FAttributeTensor>();
+					if (Output.GetDataType() != AttributeTensor.GetDataType())
+					{
+						UE_LOG(LogNNE, Warning, TEXT("Output data type %s does not match constant tensor data type %s"), 
+							   					*LogHelper::GetTensorDataTypeName(Output.GetDataType()), 
+												*LogHelper::GetTensorDataTypeName(AttributeTensor.GetDataType()));
+						return -1;
+					}
+					AttributeTensor.FillFTensorWithShapeAndData(Output);
+				} break;
+				default:
+					check(false);
+					break;
+			}
+			check(OutputTensors[0]->IsConstant());
+			return 0;
+		};
+
+		virtual bool Initialize(TConstArrayView<NNE::FTensorDesc> InputTensorDescs, TConstArrayView<NNE::FTensorDesc> OutputTensorDescs, const NNE::FAttributeMap& Attributes) override
+		{
+			check(InputTensorDescs.Num() == 0);
+			check(OutputTensorDescs.Num() == 1);
+			check(Attributes.Num() == 1);
+
+			Attribute = Attributes.GetAttributeValue(0);
+
+			return true;
+		}
+
+		virtual void Dispatch(FRDGBuilder& GraphBuilder, TConstArrayView<FTensorRDGRef> InputTensors, TConstArrayView<FTensorRDGRef> OutputTensors) override
+		{
+			checkf(false, TEXT("Dispatch should never be called, since we have a constant output"));
+		}
+	};
+
+	bool ValidateConstantOperator(const NNE::FAttributeMap& AttributeMap, TConstArrayView<ENNETensorDataType> InputTypes, TConstArrayView<NNE::FSymbolicTensorShape> InputShapes)
+	{
+		bool bIsValid = true;
+
+		FAttributeValidator AttributeValidator;
+		AttributeValidator.AddOptional(TEXT("value"), ENNEAttributeDataType::Tensor);
+		AttributeValidator.AddOptional(TEXT("value_float"), ENNEAttributeDataType::Float);
+		AttributeValidator.AddOptional(TEXT("value_floats"), ENNEAttributeDataType::FloatArray);
+		if (!AttributeValidator.Validate(AttributeMap))
+		{
+			return false;
+		}
+
+		if (AttributeMap.Num() != 1)
+		{
+			UE_LOG(LogNNE, Warning, TEXT("Constant operator requires exacly one attribute, but '%d' attributes found."), AttributeMap.Num());
+			return false;
+		}
+
+		FInputValidator InputValidator;
+		bIsValid &= InputValidator.Validate(InputTypes);
+
+		return bIsValid;
+	}
+
+	FOperatorHlsl* CreateConstantOperator()
+	{
+		return new FConstant();
+	}
+
+	bool RegisterConstantOperator(FOperatorRegistryHlsl& Registry)
+	{
+		// Note: support of a particular version is partial with respect to tensor data types (only the most typical ones are usually supported).
+		#define OP(Version) \
+		Registry.OpAdd({{TEXT("Constant"), TEXT("Onnx")}, Version}, CreateConstantOperator, ValidateConstantOperator);
+
+		OP(9)
+		OP(11)
+		OP(12)
+		OP(13)
+		OP(19)
+		OP(21)
+
+		#undef OP
+		return true;
+	}
+
+} // UE::NNERuntimeRDG::Private::Hlsl
diff --git a/Engine/Plugins/Experimental/NNERuntimeRDG/Source/NNERuntimeRDG/Private/Hlsl/NNERuntimeRDGConstant.h b/Engine/Plugins/Experimental/NNERuntimeRDG/Source/NNERuntimeRDG/Private/Hlsl/NNERuntimeRDGConstant.h
new file mode 100644
--- /dev/null
+++ b/Engine/Plugins/Experimental/NNERuntimeRDG/Source/NNERuntimeRDG/Private/Hlsl/NNERuntimeRDGConstant.h
@@ -0,0 +1,11 @@
+// Copyright Epic Games, Inc. All Rights Reserved.
+
+#pragma once
+
+#include "NNERuntimeRDGBase.h"
+#include "NNERuntimeRDGHlslOp.h"
+
+namespace UE::NNERuntimeRDG::Private::Hlsl
+{
+	bool RegisterConstantOperator(FOperatorRegistryHlsl& Registry);
+} // UE::NNERuntimeRDG::Private::Hlsl
diff --git a/Engine/Plugins/Experimental/NNERuntimeRDG/Source/NNERuntimeRDG/Private/NNERuntimeRDGHlsl.cpp b/Engine/Plugins/Experimental/NNERuntimeRDG/Source/NNERuntimeRDG/Private/NNERuntimeRDGHlsl.cpp
--- a/Engine/Plugins/Experimental/NNERuntimeRDG/Source/NNERuntimeRDG/Private/NNERuntimeRDGHlsl.cpp
+++ b/Engine/Plugins/Experimental/NNERuntimeRDG/Source/NNERuntimeRDG/Private/NNERuntimeRDGHlsl.cpp
@@ -20,6 +20,7 @@
 #include "Hlsl/NNERuntimeRDGCast.h"
 #include "Hlsl/NNERuntimeRDGConv.h"
 #include "Hlsl/NNERuntimeRDGConcat.h"
+#include "Hlsl/NNERuntimeRDGConstant.h"
 #include "Hlsl/NNERuntimeRDGConvTranspose.h"
 #include "Hlsl/NNERuntimeRDGCumSum.h"
 #include "Hlsl/NNERuntimeRDGDepthToSpace.h"
@@ -66,6 +67,7 @@ bool UNNERuntimeRDGHlslImpl::Init()
 	RegisterCastOperator(*Registry);
 	RegisterConvOperator(*Registry);
 	RegisterConcatOperator(*Registry);
+	RegisterConstantOperator(*Registry);
 	RegisterConvTransposeOperator(*Registry);
 	RegisterCumSumOperator(*Registry);
 	RegisterDepthToSpaceOperator(*Registry);
diff --git a/Engine/Plugins/Experimental/NNERuntimeRDG/Source/NNERuntimeRDGUtils/Private/NNERuntimeRDGUtilsModelOptimizerNNE.cpp b/Engine/Plugins/Experimental/NNERuntimeRDG/Source/NNERuntimeRDGUtils/Private/NNERuntimeRDGUtilsModelOptimizerNNE.cpp
--- a/Engine/Plugins/Experimental/NNERuntimeRDG/Source/NNERuntimeRDGUtils/Private/NNERuntimeRDGUtilsModelOptimizerNNE.cpp
+++ b/Engine/Plugins/Experimental/NNERuntimeRDG/Source/NNERuntimeRDGUtils/Private/NNERuntimeRDGUtilsModelOptimizerNNE.cpp
@@ -2,11 +2,15 @@
 
 #include "NNERuntimeRDGUtilsModelOptimizerNNE.h"
 
+#include "Containers/ContainersFwd.h"
 #include "NNE.h"
 #include "NNERuntimeFormat.h"
 #include "NNERuntimeRDGUtilsModelBuilderNNE.h"
 #include "NNERuntimeRDGUtilsModelOptimizerONNX.h"
 #include "NNERuntimeRDGUtilsHelpers.h"
+#include "NNETensor.h"
+#include "NNEAttributeTensor.h"
+#include "NNETypes.h"
 
 THIRD_PARTY_INCLUDES_START
 #include <onnx/onnx_pb.h>
@@ -105,7 +109,59 @@ namespace ModelOptimizerNNEHelper
 		return nullptr;
 	}
 
-	bool GetTensorInfoFromONNXInitializer(const onnx::TensorProto& Tensor, TArray<int32>& Shape, ENNETensorDataType& DataType, const void*& Data, uint64& DataSize)
+	class TensorInfoData
+	{
+		const uint8* Data = nullptr;
+		size_t DataSize = 0;
+		uint32 ElementSize = 0;
+		uint32 ElementStride = 0;
+		TArray<uint8_t> PackedData;
+
+	public:
+
+		TensorInfoData() {}
+
+		static TensorInfoData Make(const void* Data, size_t DataSize, uint32 ElementSize, uint32 ElementStride)
+		{
+			check(Data);
+			check(ElementStride > 0);
+			check(DataSize % ElementStride == 0);
+			check(ElementSize <= ElementStride);
+
+			TensorInfoData Result = TensorInfoData();
+			Result.Data = reinterpret_cast<const uint8 *>(Data);
+			Result.DataSize = DataSize;
+			Result.ElementSize = ElementSize;
+			Result.ElementStride = ElementStride;
+
+			return Result;
+		}
+
+		TConstArrayView<uint8> GetArrayView()
+		{
+			if (ElementStride > ElementSize)
+			{
+				if (PackedData.IsEmpty())
+				{
+					PackedData.Reset(DataSize / ElementStride * ElementSize);
+					for (uint64 DataIndex = 0; DataIndex < DataSize; DataIndex += ElementStride)
+					{
+						for (uint64 Offset = 0; Offset < ElementSize; ++Offset)
+						{
+							PackedData.Add(Data[DataIndex + Offset]);
+						}
+					}
+				}
+				return MakeArrayView(PackedData);
+			}
+			else
+			{
+				return MakeArrayView(Data, DataSize);
+			}
+		}
+	};
+
+	bool GetTensorInfoFromONNXInitializer(const onnx::TensorProto& Tensor, TArray<int32>& Shape, ENNETensorDataType& DataType, TensorInfoData& OutData)
 	{
 		DataType = GetNNETensorTypeFromONNX(Tensor.data_type());
 
@@ -115,17 +171,24 @@ namespace ModelOptimizerNNEHelper
 			Shape.Add(static_cast<int32>(Tensor.dims(i)));
 		}
 
+		const uint32 ElementSize = UE::NNE::GetTensorDataTypeSizeInBytes(DataType);
+		const void *Data = nullptr;
+		size_t DataSize = 0;
+		uint32 ElementStride;
+
 		if (Tensor.has_raw_data())
 		{
 			const std::string& RawData = Tensor.raw_data();
 			Data = RawData.c_str();
 			DataSize = RawData.size();
+			ElementStride = ElementSize;
 		}
 		else if (Tensor.double_data().size())
 		{
 			//DOUBLE or COMPLEX128
 			Data = Tensor.double_data().data();
 			DataSize = Tensor.double_data().size() * sizeof(double);
+			ElementStride = sizeof(double);
 		}
 		else if (Tensor.external_data().size()) 
 		{
@@ -137,30 +200,40 @@ namespace ModelOptimizerNNEHelper
 			//FLOAT or COMPLEX64
 			Data = Tensor.float_data().data();
 			DataSize = Tensor.float_data().size() * sizeof(float);
+			ElementStride = sizeof(float);
 		}
-		else if (Tensor.int32_data().size() && DataType == ENNETensorDataType::Int32)
+		else if (Tensor.int32_data().size())
 		{
-			//Supported : INT32
-			//Not supported at the moment: INT16, INT8, UINT16, UINT8, BOOL, FLOAT16, BFLOAT16, FLOAT8E4M3FN, FLOAT8E4M3FNUZ, FLOAT8E5M2, FLOAT8E5M2FNUZ, UINT32
+			//Supported : INT32, FLOAT16
+			//Not supported at the moment: INT16, INT8, UINT16, UINT8, BOOL, BFLOAT16, FLOAT8E4M3FN, FLOAT8E4M3FNUZ, FLOAT8E5M2, FLOAT8E5M2FNUZ, UINT32
+			if (DataType != ENNETensorDataType::Int32 &&
+				DataType != ENNETensorDataType::Half)
+			{
+				return false;
+			}
 			Data = Tensor.int32_data().data();
 			DataSize = Tensor.int32_data().size() * sizeof(int32);
+			ElementStride = sizeof(int32);
 		}
 		else if (Tensor.uint64_data().size())
 		{
 			//Supported UINT64
 			Data = Tensor.uint64_data().data();
 			DataSize = Tensor.uint64_data().size() * sizeof(uint64);
+			ElementStride = sizeof(uint64);
 		}
 		else if (Tensor.int64_data().size())
 		{
 			//Supported INT64
 			Data = Tensor.int64_data().data();
 			DataSize = Tensor.int64_data().size() * sizeof(int64);
+			ElementStride = sizeof(int64);
 		}
 		else
 		{
 			return false;
 		}
+		OutData = TensorInfoData::Make(Data, DataSize, ElementSize, ElementStride);
 		return true;
 	}
 
@@ -178,6 +251,48 @@ namespace ModelOptimizerNNEHelper
 		}
 	}
 
+	static bool GetAttributeTensorFromONNXInitializer(const onnx::TensorProto& InTensor, NNE::Internal::FAttributeTensor& OutTensor)
+	{
+		using NNE::FTensorShape;
+		using NNE::Internal::FAttributeTensor;
+
+		ENNETensorDataType DataType = ENNETensorDataType::None;
+		TArray<int32> Shape;
+		TensorInfoData Data;
+
+		if (!GetTensorInfoFromONNXInitializer(InTensor, Shape, DataType, Data))
+		{
+			UE_LOG(LogNNE, Error, TEXT("Tensor data could not be loaded"));
+			return false;
+		}
+		if (Shape.Num() > FTensorShape::MaxRank)
+		{
+			UE_LOG(LogNNE, Error, TEXT("Tensor shape of rank %i exceeds MaxRank %i"), Shape.Num(), FTensorShape::MaxRank);
+			return false;
+		}
+		TArray<uint32, TInlineAllocator<FTensorShape::MaxRank>> UIntShape;
+		for (int32 value : Shape)
+		{
+			if (value < 0)
+			{
+				UE_LOG(LogNNE, Error, TEXT("Tensor shape has negative value"));
+				return false;
+			}
+			UIntShape.Add(value);
+		}
+		const FTensorShape TensorShape = FTensorShape::Make(UIntShape);
+		TConstArrayView<uint8> DataView = Data.GetArrayView();
+		uint32 ExpectedDataSize = TensorShape.Volume() * UE::NNE::GetTensorDataTypeSizeInBytes(DataType);
+		if (ExpectedDataSize != DataView.NumBytes())
+		{
+			UE_LOG(LogNNE, Warning, TEXT("Tensor data size %i doesn't match expected data size %i"), 
+				   					DataView.NumBytes(), ExpectedDataSize);
+			return false;
+		}
+		OutTensor = FAttributeTensor::Make(TensorShape, DataType, DataView);
+		return true;
+	}
+
 	bool BuildNNEFormatFromONNX(TArray<uint8>& ONNXData, TArray<uint8>& NNEData)
 	{
 		TUniquePtr<IModelBuilder> Builder = CreateNNEModelBuilder();
@@ -234,12 +349,11 @@ namespace ModelOptimizerNNEHelper
 			const onnx::TensorProto* Initializer = GetInitializerFromGraphProto(Graph, Output.name());
 			if (Initializer)
 			{
-				const void* Data = nullptr;
-				uint64 DataSize = 0;
 				ENNETensorDataType InitializerDataType;
 				TArray<int32> InitializerShape;
+				TensorInfoData Data;
 
-				if (!GetTensorInfoFromONNXInitializer(*Initializer, InitializerShape, InitializerDataType, Data, DataSize))
+				if (!GetTensorInfoFromONNXInitializer(*Initializer, InitializerShape, InitializerDataType, Data))
 				{
 					UE_LOG(LogNNE, Error, TEXT("Tensor data could not be loaded for weights of output node '%s'"), ANSI_TO_TCHAR(Output.name().c_str()));
 					return false;
@@ -258,8 +372,9 @@ namespace ModelOptimizerNNEHelper
 					return false;
 				}
 
+				TConstArrayView<uint8> DataView = Data.GetArrayView();
 				IModelBuilder::FHTensor TensorInitializer =
-					Builder->AddConstantTensor(FString(ANSI_TO_TCHAR(Output.name().c_str())) + TEXT("_NNEInitializer"), DataType, Shape, Data, DataSize);
+					Builder->AddConstantTensor(FString(ANSI_TO_TCHAR(Output.name().c_str())) + TEXT("_NNEInitializer"), DataType, Shape, DataView.GetData(), DataView.NumBytes());
 
 				const FString IdentityOpType = TEXT("Identity");
 				TOptional<uint32> OpVersion = GetOpVersionFromOpsetVersion(IdentityOpType, (int) ModelProto.opset_import(0).version());
@@ -344,6 +459,38 @@ namespace ModelOptimizerNNEHelper
 
 					Builder->AddOperatorAttribute(Op, AttributeName, FNNEAttributeValue(Values));
 				}
+				else if (Attribute.type() == onnx::AttributeProto::TENSOR)
+				{
+					using NNE::Internal::FAttributeTensor;
+
+					FAttributeTensor Tensor;
+					if (!GetAttributeTensorFromONNXInitializer(Attribute.t(), Tensor))
+					{
+						UE_LOG(LogNNE, Error, TEXT("Failed to get data from attribute %s in operator %s"), *AttributeName, *NNEOpName);
+						return false;
+					}
+					Builder->AddOperatorAttribute(Op, AttributeName, FNNEAttributeValue(Tensor));
+				}
+				else if (Attribute.type() == onnx::AttributeProto::TENSORS)
+				{
+					using NNE::Internal::FAttributeTensor;
+
+					TArray<FAttributeTensor> Values;
+
+					int TensorIndex = 0;
+					for (const onnx::TensorProto& Initializer : Attribute.tensors())
+					{
+						FAttributeTensor Tensor;
+						if (!GetAttributeTensorFromONNXInitializer(Initializer, Tensor))
+						{
+							UE_LOG(LogNNE, Error, TEXT("Failed to get data from tensor at index %i of attribute %s in operator %s"), TensorIndex, *AttributeName, *NNEOpName);
+							return false;
+						}
+						Values.Add(Tensor);
+						TensorIndex++;
+					}
+					Builder->AddOperatorAttribute(Op, AttributeName, FNNEAttributeValue(Values));
+				}
 				else
 				{
 					//Note: Would be good to have better error reporting by adding type (example: sparse tensor)
@@ -355,20 +502,20 @@ namespace ModelOptimizerNNEHelper
 			{
 				ENNETensorDataType DataType = ENNETensorDataType::None;
 				TArray<int32> Shape;
-				const void* Data = nullptr;
-				uint64 DataSize = 0;
+				TensorInfoData Data;
 
 				IModelBuilder::FHTensor Tensor;
 
 				const onnx::TensorProto* Initializer = GetInitializerFromGraphProto(Graph, TensorName);
 				if (Initializer)
 				{
-					if (!GetTensorInfoFromONNXInitializer(*Initializer, Shape, DataType, Data, DataSize))
+					if (!GetTensorInfoFromONNXInitializer(*Initializer, Shape, DataType, Data))
 					{
 						UE_LOG(LogNNE, Error, TEXT("Tensor data could not be loaded for weight '%s' in node '%s' of type '%s'"), ANSI_TO_TCHAR(TensorName.c_str()), *NNEOpName, *NNEOpType);
 						return false;
 					}
-					Tensor = Builder->AddConstantTensor(ANSI_TO_TCHAR(TensorName.c_str()), DataType, Shape, Data, DataSize);
+					TConstArrayView<uint8> DataView = Data.GetArrayView();
+					Tensor = Builder->AddConstantTensor(ANSI_TO_TCHAR(TensorName.c_str()), DataType, Shape, DataView.GetData(), DataView.NumBytes());
 				}
 				else if (!TensorName.empty())
 				{
diff --git a/Engine/Plugins/Experimental/NNERuntimeRDG/Source/NNERuntimeRDGUtils/Private/NNERuntimeRDGUtilsModelOptimizerONNX.cpp b/Engine/Plugins/Experimental/NNERuntimeRDG/Source/NNERuntimeRDGUtils/Private/NNERuntimeRDGUtilsModelOptimizerONNX.cpp
--- a/Engine/Plugins/Experimental/NNERuntimeRDG/Source/NNERuntimeRDGUtils/Private/NNERuntimeRDGUtilsModelOptimizerONNX.cpp
+++ b/Engine/Plugins/Experimental/NNERuntimeRDG/Source/NNERuntimeRDGUtils/Private/NNERuntimeRDGUtilsModelOptimizerONNX.cpp
@@ -41,11 +41,27 @@ public:
 			return false;
 		}
 
+		onnx::ModelProto ModelProto;
+		const bool result = ModelProto.ParseFromArray(Model.Data.GetData(), Model.Data.Num());
+		if (!result)
+		{
+			UE_LOG(LogNNE, Warning, TEXT("%s could not parse the input model as a ModelProto."), *(GetName()));
+			return false;
+		}
+
 		static const auto CVarHlslModelOptimization = IConsoleManager::Get().FindTConsoleVariableDataInt(TEXT("nne.hlsl.ModelOptimization"));
 		if (CVarHlslModelOptimization && CVarHlslModelOptimization->GetValueOnAnyThread() == 0)
 		{
 			return true;
 		}
+		
+		// We disable ONNX Runtime optimizations if a model uses FP16 Tensors 
+		// since these optimizations would add cast operators from and to FP16
+		// at the beginning and end of the network and convert all other operators to FP32.
+		if (HasFP16Tensor(ModelProto))
+		{
+			return true;
+		}
 
 		FString ProjIntermediateDir = FPaths::ConvertRelativePathToFull(FPaths::ProjectIntermediateDir());
 		FString ModelToOptimizePath = FPaths::CreateTempFilename(*ProjIntermediateDir, TEXT("ORTOptimizerPass_ToOptimize"), TEXT(".onnx"));
@@ -86,6 +102,45 @@ public:
 
 		return true;
 	}
+private:
+	bool HasFP16Tensor(const onnx::ModelProto& Model) const
+	{
+		const onnx::GraphProto& Graph = Model.graph();
+		for (const onnx::TensorProto& Tensor : Graph.initializer())
+		{
+			if (Tensor.data_type() == ONNXTensorElementDataType::ONNX_TENSOR_ELEMENT_DATA_TYPE_FLOAT16)
+			{
+				return true;
+			}
+		}
+
+		for (onnx::ValueInfoList Tensors : {Graph.input(), Graph.output()})
+		{
+			for (const onnx::ValueInfoProto& Tensor : Tensors)
+			{
+				if (!Tensor.has_type())
+				{
+					continue;
+				}
+				const onnx::TypeProto Type = Tensor.type();
+				if (!Type.has_tensor_type())
+				{
+					continue;
+				}
+				const onnx::TypeProto_Tensor TensorType = Type.tensor_type();
+				if (!TensorType.has_elem_type())
+				{
+					continue;
+				}
+				if (TensorType.elem_type() == ONNXTensorElementDataType::ONNX_TENSOR_ELEMENT_DATA_TYPE_FLOAT16)
+				{
+					return true;
+				}
+			}
+		}
+
+		return false;
+	}
 
 };
 
@@ -128,7 +183,9 @@ public:
 			const char* DomainPtr = OpSet.domain().c_str();
 			FString Domain = DomainPtr;
 			bool IsUsed = UsedDomains.Contains(Domain);
-			if (IsUsed)
+			// We additionally add the OpSet from models that don't have any UsedDomains/Nodes
+			// since we otherwise would get an invalid model
+			if (IsUsed || UsedDomains.IsEmpty())
 			{
 				UsedOperatorSet.Add()->CopyFrom(OpSet);
 			}
diff --git a/Engine/Source/Runtime/NNE/Internal/NNEAttributeDataType.h b/Engine/Source/Runtime/NNE/Internal/NNEAttributeDataType.h
--- a/Engine/Source/Runtime/NNE/Internal/NNEAttributeDataType.h
+++ b/Engine/Source/Runtime/NNE/Internal/NNEAttributeDataType.h
@@ -20,5 +20,7 @@ enum class ENNEAttributeDataType : uint8
 	Int32,								//!< 32-bit signed integer
 	Int32Array,							//!< TArray of 32-bit signed integers
 	String,								//!< built-in FString
-	StringArray							//!< TArray of built-in FString
+	StringArray,						//!< TArray of built-in FString
+	Tensor,								//!< FAttributeTensor
+	TensorArray							//!< TArray of FAttributeTensor
 };
\ No newline at end of file
diff --git a/Engine/Source/Runtime/NNE/Internal/NNEAttributeTensor.h b/Engine/Source/Runtime/NNE/Internal/NNEAttributeTensor.h
new file mode 100644
--- /dev/null
+++ b/Engine/Source/Runtime/NNE/Internal/NNEAttributeTensor.h
@@ -0,0 +1,57 @@
+// Copyright Epic Games, Inc. All Rights Reserved.
+
+#pragma once
+
+#include "Containers/ArrayView.h"
+#include "NNETensor.h"
+#include "NNETypes.h"
+#include "Serialization/Archive.h"
+
+namespace UE::NNE::Internal
+{
+	class FAttributeTensor
+	{
+	protected:
+		ENNETensorDataType	DataType;
+		TArray<uint32, TInlineAllocator<FTensorShape::MaxRank>> Shape;
+		TArray<uint8>		Data;
+
+	public:
+
+		ENNETensorDataType GetDataType() const
+		{
+			return DataType;
+		}
+
+		void FillFTensorWithShapeAndData(FTensor& Tensor) const
+		{
+			Tensor.SetShape(FTensorShape::Make(Shape));
+			Tensor.SetPreparedData(MakeConstArrayView(Data));
+		}
+
+		static FAttributeTensor Make(const FTensorShape& Shape,
+									 ENNETensorDataType DataType, 
+									 TConstArrayView<uint8> Data)
+		{
+			uint64 Volume = Shape.Volume();
+			check(Volume <= TNumericLimits<uint32>::Max());
+			check(Data.NumBytes() == GetTensorDataTypeSizeInBytes(DataType) * Volume);
+
+			FAttributeTensor Tensor;
+			Tensor.DataType = DataType;
+			Tensor.Shape.Reset();
+			Tensor.Shape.Append(Shape.GetData());
+			Tensor.Data.Reset();
+			Tensor.Data.Append(Data);
+			return Tensor;
+		}
+
+		friend FArchive& operator<<(FArchive& Ar, FAttributeTensor& Tensor)
+		{
+			Ar << Tensor.DataType;
+			Ar << Tensor.Shape;
+			Ar << Tensor.Data;
+			return Ar;
+		}
+	};
+} // namespace UE::NNE::Internal
diff --git a/Engine/Source/Runtime/NNE/Internal/NNEAttributeValueTraits.h b/Engine/Source/Runtime/NNE/Internal/NNEAttributeValueTraits.h
--- a/Engine/Source/Runtime/NNE/Internal/NNEAttributeValueTraits.h
+++ b/Engine/Source/Runtime/NNE/Internal/NNEAttributeValueTraits.h
@@ -3,6 +3,7 @@
 #pragma once
 
 #include "NNEAttributeDataType.h"
+#include "NNEAttributeTensor.h"
 
 template<typename T> struct TNNEAttributeValueTraits
 {
@@ -29,6 +30,11 @@ template<> struct TNNEAttributeValueTraits<FString>
 	static constexpr ENNEAttributeDataType GetType() { return ENNEAttributeDataType::String; }
 };
 
+template<> struct TNNEAttributeValueTraits<UE::NNE::Internal::FAttributeTensor>
+{
+	static constexpr ENNEAttributeDataType GetType() { return ENNEAttributeDataType::Tensor; }
+};
+
 template<> struct TNNEAttributeValueTraits<TArray<int32>>
 {
 	static constexpr ENNEAttributeDataType GetType() { return ENNEAttributeDataType::Int32Array; }
@@ -42,4 +48,9 @@ template<> struct TNNEAttributeValueTraits<TArray<float>>
 template<> struct TNNEAttributeValueTraits<TArray<FString>>
 {
 	static constexpr ENNEAttributeDataType GetType() { return ENNEAttributeDataType::StringArray; }
-};
\ No newline at end of file
+};
+
+template<> struct TNNEAttributeValueTraits<TArray<UE::NNE::Internal::FAttributeTensor>>
+{
+	static constexpr ENNEAttributeDataType GetType() { return ENNEAttributeDataType::TensorArray; }
+};
