NNE - Fix Conv negative autopad

Problems:
* Formula to compute autopads for Hlsl Conv can yield negative pads.

Solutions:
* Correct the formula to have >= 0 computed pads always.

Tests:
Win64 Editor NNEQA tests: all

#rnx

#virtualized

diff --git a/Engine/Plugins/Experimental/NNERuntimeRDG/Source/NNEHlslShaders/Private/NNEHlslShadersConvCS.cpp b/Engine/Plugins/Experimental/NNERuntimeRDG/Source/NNEHlslShaders/Private/NNEHlslShadersConvCS.cpp
--- a/Engine/Plugins/Experimental/NNERuntimeRDG/Source/NNEHlslShaders/Private/NNEHlslShadersConvCS.cpp
+++ b/Engine/Plugins/Experimental/NNERuntimeRDG/Source/NNEHlslShaders/Private/NNEHlslShadersConvCS.cpp
@@ -151,7 +151,9 @@ namespace UE::NNEHlslShaders::Internal
 		for (int32 DimensionIndex = 0; DimensionIndex < NumDimensions; DimensionIndex++)
 		{
 			int32 DilatedKernelSize = (DimensionIndex < Dilations.Num() ? Dilations[DimensionIndex] : 1) * (WShape[DimensionIndex + 2] - 1) + 1;
-			int32 TotalPad = ((int32)(((XShape[DimensionIndex + 2] + (DimensionIndex < Strides.Num() ? Strides[DimensionIndex] : 1) - 1) / (DimensionIndex < Strides.Num() ? Strides[DimensionIndex] : 1)) - 1)) * (DimensionIndex < Strides.Num() ? Strides[DimensionIndex] : 1) + DilatedKernelSize - XShape[DimensionIndex + 2];
+			int32 LastOutputIdx = ((int32) XShape[DimensionIndex + 2] + (DimensionIndex < Strides.Num() ? Strides[DimensionIndex] : 1) - 1) / (DimensionIndex < Strides.Num() ? Strides[DimensionIndex] : 1) - 1;
+			int32 TotalPad = LastOutputIdx * (DimensionIndex < Strides.Num() ? Strides[DimensionIndex] : 1) + DilatedKernelSize - XShape[DimensionIndex + 2];
+			TotalPad = TotalPad >= 0 ? TotalPad : 0;
 			if (AutoPad == EConvAutoPad::SAME_LOWER)
 			{
 				Result[DimensionIndex] = (TotalPad + 1) / 2;
