[NNERuntimeRDG HLSL] Add Convolution variant based on Winograd

Problem:  There is room for performance improvements for convolutions that have a 3x3 kernel. Since this is a common kernel size it's important to optimize for this.

Solution: Implement a new convolution variant based on Winograd that gets used when we have a 3x3 kernel.

Tested:
* Win64 Standalone tests: all
* Win64 Editor tests: all
* Consoles tests: Convolution, OIDN and Smoke

#rnx

#virtualized

diff --git a/Engine/Plugins/Experimental/NNERuntimeRDG/Shaders/Private/NNEHlslShaders/NNEHlslShadersConvWinogradInput.usf b/Engine/Plugins/Experimental/NNERuntimeRDG/Shaders/Private/NNEHlslShaders/NNEHlslShadersConvWinogradInput.usf
new file mode 100644
--- /dev/null
+++ b/Engine/Plugins/Experimental/NNERuntimeRDG/Shaders/Private/NNEHlslShaders/NNEHlslShadersConvWinogradInput.usf
@@ -0,0 +1,95 @@
+// Copyright Epic Games, Inc. All Rights Reserved.
+
+#include "/Engine/Public/Platform.ush"
+#include "NNEHlslShadersBroadcastHelper.ush"
+
+Buffer<float> Input; // Ni x C x H x W
+RWBuffer<float> Output; // Ni x 36 x C x Ceil(H/4)*Ceil(W/4)
+
+int C;
+int H;
+int W;
+int WBlockCount;
+
+int CInputStride;
+int HInputStride;
+
+int NiOutputStride;
+int MatrixOutputStride;
+int COutputStride;
+int HOutputStride;
+
+[numthreads(THREADGROUP_SIZE_X, 1, 1)]
+void ConvWinogradInput(
+	in const uint3 GroupID : SV_GroupID,
+	in const uint3 DispatchThreadID : SV_DispatchThreadID,
+	in const uint3 GroupThreadID : SV_GroupThreadID)
+{
+	const int WBlock = DispatchThreadID.x;
+	if (WBlock < WBlockCount)
+	{
+		const int WBase = WBlock * 4 - 1;
+		const int Scalar_HBase = GroupID.y * 4 - 1;
+
+		const int Scalar_InputBase = GroupID.z * CInputStride;
+
+		int Scalar_OutputC, Scalar_OutputNi;
+		DivMod(GroupID.z, C, Scalar_OutputNi, Scalar_OutputC);
+		const int Scalar_OutputBase = Scalar_OutputNi * NiOutputStride + Scalar_OutputC * COutputStride + GroupID.y * HOutputStride;
+		const int OutputOffset = Scalar_OutputBase + WBlock;
+
+		// Load d
+		float d[6][6];
+		UNROLL
+		for (int HOffset = 0; HOffset < 6; HOffset++)
+		{
+			const int h = Scalar_HBase + HOffset;
+			const bool ValidH = h >= 0 && h < H;
+			const int HInput = Scalar_InputBase + clamp(h, 0, H - 1) * HInputStride;
+
+			UNROLL
+			for (int WOffset = 0; WOffset < 6; WOffset++)
+			{
+				int w = WBase + WOffset;
+				const bool ValidInput = ValidH && w >= 0 && w < W;
+				w = clamp(w, 0, W - 1);
+				const float Value = Input[HInput + w];
+				d[HOffset][WOffset] = ValidInput ? Value : 0;
+			}
+		}
+		// Calculate B^TdB
+		float BTd[6][6];
+		UNROLL
+		for (int i = 0; i < 6; i++)
+		{
+			BTd[0][i] =  d[0][i] * 4 - d[2][i] * 5 + d[4][i];
+			BTd[1][i] = -d[1][i] * 4 - d[2][i] * 4 + d[3][i] + d[4][i];
+			BTd[2][i] =  d[1][i] * 4 - d[2][i] * 4 - d[3][i] + d[4][i];
+			BTd[3][i] = -d[1][i] * 2 - d[2][i] + d[3][i] * 2 + d[4][i];
+			BTd[4][i] =  d[1][i] * 2 - d[2][i] - d[3][i] * 2 + d[4][i];
+			BTd[5][i] =  d[1][i] * 4 - d[3][i] * 5 + d[5][i];
+		}
+		float BTdB[6][6];
+		UNROLL
+		for (int i = 0; i < 6; i++)
+		{
+			BTdB[i][0] =  BTd[i][0] * 4 - BTd[i][2] * 5 + BTd[i][4];
+			BTdB[i][1] = -BTd[i][1] * 4 - BTd[i][2] * 4 + BTd[i][3] + BTd[i][4];
+			BTdB[i][2] =  BTd[i][1] * 4 - BTd[i][2] * 4 - BTd[i][3] + BTd[i][4];
+			BTdB[i][3] = -BTd[i][1] * 2 - BTd[i][2] + BTd[i][3] * 2 + BTd[i][4];
+			BTdB[i][4] =  BTd[i][1] * 2 - BTd[i][2] - BTd[i][3] * 2 + BTd[i][4];
+			BTdB[i][5] =  BTd[i][1] * 4 - BTd[i][3] * 5 + BTd[i][5];
+		}
+
+		// Store result to Output
+		UNROLL
+		for (int j = 0; j < 6; j++)
+		{
+			UNROLL
+			for (int i = 0; i < 6; i++)
+			{
+				Output[(j * 6 + i) * MatrixOutputStride + OutputOffset] = BTdB[j][i];
+			}
+		}
+	}
+}
diff --git a/Engine/Plugins/Experimental/NNERuntimeRDG/Shaders/Private/NNEHlslShaders/NNEHlslShadersConvWinogradMMM.usf b/Engine/Plugins/Experimental/NNERuntimeRDG/Shaders/Private/NNEHlslShaders/NNEHlslShadersConvWinogradMMM.usf
new file mode 100644
--- /dev/null
+++ b/Engine/Plugins/Experimental/NNERuntimeRDG/Shaders/Private/NNEHlslShaders/NNEHlslShadersConvWinogradMMM.usf
@@ -0,0 +1,144 @@
+// Copyright Epic Games, Inc. All Rights Reserved.
+
+#include "/Engine/Public/Platform.ush"
+
+#define K_STEP 16
+#define BLOCK_ELEM_COUNT_M 64
+#define VECTOR_SIZE 2
+
+#define BLOCK_VECTOR_ELEM_COUNT_M (BLOCK_ELEM_COUNT_M / VECTOR_SIZE)
+#define BLOCK_VECTOR_ELEM_COUNT_N (BLOCK_ELEM_COUNT_N / VECTOR_SIZE)
+
+#define BLOCK_COUNT_M (BLOCK_ELEM_COUNT_M / 4)
+#define BLOCK_COUNT_N (BLOCK_ELEM_COUNT_N / 4)
+#define BLOCK_COUNT   (BLOCK_COUNT_M * BLOCK_COUNT_N)
+
+#define THREADGROUP_SIZE_X BLOCK_VECTOR_ELEM_COUNT_M
+#define THREADGROUP_SIZE_Y (BLOCK_COUNT / THREADGROUP_SIZE_X)
+#define THREADGROUP_SIZE (THREADGROUP_SIZE_X * THREADGROUP_SIZE_Y)
+
+#define WORK_TYPE_VECTOR vector <WORK_TYPE, VECTOR_SIZE>
+
+Buffer<WORK_TYPE_VECTOR> Input; // Ni x 36 x K x M
+Buffer<WORK_TYPE_VECTOR> Weight; // 36 x K x N 
+RWBuffer<WORK_TYPE_VECTOR> Output; // Ni x 36 x N x M
+
+groupshared WORK_TYPE_VECTOR LDS_InputOutput[max(K_STEP, BLOCK_ELEM_COUNT_N)][BLOCK_VECTOR_ELEM_COUNT_M];
+groupshared WORK_TYPE_VECTOR LDS_Weight[K_STEP][BLOCK_VECTOR_ELEM_COUNT_N];
+
+int M;
+int N;
+int K;
+
+int MatrixInputStride;
+int KInputStride;
+
+int MatrixWeightStride;
+int KWeightStride;
+
+int MatrixOutputStride;
+int NOutputStride;
+
+[numthreads(THREADGROUP_SIZE_X, THREADGROUP_SIZE_Y, 1)]
+void ConvWinogradMMM(
+	in const uint3 DispatchThreadID : SV_DispatchThreadID,
+	in const uint3 GroupID : SV_GroupID,
+	in const uint3 GroupThreadID : SV_GroupThreadID,
+	in const uint GroupIndex : SV_GroupIndex)
+{
+	const int Scalar_NBase = GroupID.y * BLOCK_VECTOR_ELEM_COUNT_N;
+	const int Scalar_InputMatrix = GroupID.z * MatrixInputStride;
+	const int Scalar_WeightMatrix = (GroupID.z % 36) * MatrixWeightStride;
+	const int Scalar_OutputMatrix = GroupID.z * MatrixOutputStride;
+	const int Scalar_InputIndexMax = Scalar_InputMatrix + MatrixInputStride - 1;
+	const int Scalar_WeightIndexMax = Scalar_WeightMatrix + MatrixWeightStride - 1;
+	const int Scalar_OutputIndexMax = Scalar_OutputMatrix + MatrixOutputStride - 1;
+
+	const int MBlock = GroupIndex % BLOCK_COUNT_M;
+	const int MBlockBase = MBlock * (4 / VECTOR_SIZE);
+	const int NBlock = GroupIndex / BLOCK_COUNT_M;
+	const int NBlockBase = NBlock  * (4 / VECTOR_SIZE);
+
+	matrix < WORK_TYPE, 4, 4 > Values = 0;
+
+	const int MInputIndex = DispatchThreadID.x;
+	int InputIndex = Scalar_InputMatrix;
+	InputIndex += GroupThreadID.y * KInputStride;
+	InputIndex += MInputIndex < M ? MInputIndex : MatrixInputStride;
+
+	const int NWeightIndex = Scalar_NBase + GroupIndex % BLOCK_VECTOR_ELEM_COUNT_N;
+	int WeightIndex = Scalar_WeightMatrix;
+	WeightIndex += (GroupIndex / BLOCK_VECTOR_ELEM_COUNT_N) * KWeightStride;
+	WeightIndex += NWeightIndex < N ? NWeightIndex : MatrixWeightStride;
+
+	for (int KBase = 0; KBase < K; KBase += K_STEP)
+	{
+		//Load Input and Weights
+		UNROLL
+		for (int KOffset = 0; KOffset < K_STEP; KOffset += THREADGROUP_SIZE_Y)
+		{
+			const int InputIndexClamped = min(InputIndex, Scalar_InputIndexMax);
+			WORK_TYPE_VECTOR InputValue = Input[InputIndexClamped];
+			InputValue = InputIndex <= Scalar_InputIndexMax ? InputValue : 0;
+			LDS_InputOutput[KOffset + GroupThreadID.y][GroupThreadID.x] = InputValue;
+
+			InputIndex += THREADGROUP_SIZE_Y * KInputStride;
+		}
+		UNROLL
+		for (int KOffset = 0; KOffset < K_STEP; KOffset += (THREADGROUP_SIZE / BLOCK_VECTOR_ELEM_COUNT_N))
+		{
+			const int WeightIndexClamped = min(WeightIndex, Scalar_WeightIndexMax);
+			WORK_TYPE_VECTOR WeightValue = Weight[WeightIndexClamped];
+			WeightValue = WeightIndex <= Scalar_WeightIndexMax ? WeightValue : 0;
+			LDS_Weight[KOffset + GroupIndex / BLOCK_VECTOR_ELEM_COUNT_N][GroupIndex % BLOCK_VECTOR_ELEM_COUNT_N] = WeightValue;
+
+			WeightIndex += (THREADGROUP_SIZE / BLOCK_VECTOR_ELEM_COUNT_N) * KWeightStride;
+		}
+
+		GroupMemoryBarrierWithGroupSync();
+
+		ISOLATE
+		{
+			//Multiply
+			UNROLL
+			for (int KIndex = 0; KIndex < K_STEP; KIndex++)
+			{
+				matrix < WORK_TYPE, 1, 4 > Reg_Input = { LDS_InputOutput[KIndex][MBlockBase], LDS_InputOutput[KIndex][MBlockBase + 1] };
+				matrix < WORK_TYPE, 4, 1 > Reg_Weight = { LDS_Weight[KIndex][NBlockBase], LDS_Weight[KIndex][NBlockBase + 1] };
+
+				Values += mul(Reg_Weight, Reg_Input);
+			}
+		}
+
+		GroupMemoryBarrierWithGroupSync();
+	}
+
+	// Move values from Reg to LDS
+	UNROLL
+	for (int NOffset = 0; NOffset < 4; NOffset++)
+	{
+		const int NLDSIndex = NBlock * 4 + NOffset;
+		LDS_InputOutput[NLDSIndex][MBlockBase + 0] = Values[NOffset].xy;
+		LDS_InputOutput[NLDSIndex][MBlockBase + 1] = Values[NOffset].zw;
+	}
+
+	
+	GroupMemoryBarrierWithGroupSync();
+
+	// Move values from LDS to VRAM
+	const int MOutputIndex = DispatchThreadID.x;
+	int OutputIndex = Scalar_OutputMatrix;
+	OutputIndex += Scalar_NBase * VECTOR_SIZE * NOutputStride;
+	OutputIndex += GroupThreadID.y * NOutputStride;
+	OutputIndex += MOutputIndex < M ? MOutputIndex : MatrixOutputStride;
+
+	UNROLL
+	for (int NOffset = 0; NOffset < BLOCK_ELEM_COUNT_N; NOffset += THREADGROUP_SIZE_Y)
+	{
+		if (OutputIndex <= Scalar_OutputIndexMax)
+		{
+			Output[OutputIndex] = LDS_InputOutput[NOffset + GroupThreadID.y][GroupThreadID.x];
+		}
+		OutputIndex += THREADGROUP_SIZE_Y * NOutputStride;
+	}
+}
diff --git a/Engine/Plugins/Experimental/NNERuntimeRDG/Shaders/Private/NNEHlslShaders/NNEHlslShadersConvWinogradOutput.usf b/Engine/Plugins/Experimental/NNERuntimeRDG/Shaders/Private/NNEHlslShaders/NNEHlslShadersConvWinogradOutput.usf
new file mode 100644
--- /dev/null
+++ b/Engine/Plugins/Experimental/NNERuntimeRDG/Shaders/Private/NNEHlslShaders/NNEHlslShadersConvWinogradOutput.usf
@@ -0,0 +1,95 @@
+// Copyright Epic Games, Inc. All Rights Reserved.
+
+#include "/Engine/Public/Platform.ush"
+#include "NNEHlslShadersBroadcastHelper.ush"
+
+Buffer<WORK_TYPE> Input; // Ni x 36 x CExtended x Ceil(H/4)*Ceil(W/4)
+Buffer<WORK_TYPE> Bias; // C
+RWBuffer<WORK_TYPE> Output; // Ni x C x H x W
+
+int C;
+int H;
+int W;
+int WBlockCount;
+
+int NiInputStride;
+int MatrixInputStride;
+int CInputStride;
+int HInputStride;
+
+int COutputStride;
+int HOutputStride;
+
+[numthreads(THREADGROUP_SIZE_X, 1, 1)]
+void ConvWinogradOutput(
+	in const uint3 DispatchThreadID : SV_DispatchThreadID,
+	in const uint3 GroupID : SV_GroupID,
+	in const uint3 GroupThreadID : SV_GroupThreadID)
+{
+	const int WBlock = DispatchThreadID.x;
+	if (WBlock < WBlockCount)
+	{
+		int Scalar_C, Scalar_Ni;
+		DivMod(GroupID.z, C, Scalar_Ni, Scalar_C);
+
+		const int InputOffset = Scalar_Ni * NiInputStride + Scalar_C * CInputStride + GroupID.y * HInputStride + WBlock;
+		const int Scalar_OutputBase = GroupID.z * COutputStride + (GroupID.y * HOutputStride + WBlock) * 4;
+
+		// Load block from Input
+		WORK_TYPE I[6][6];
+		UNROLL
+		for (int j = 0; j < 6; j++)
+		{
+			UNROLL
+			for (int i = 0; i < 6; i++)
+			{
+				I[j][i] = Input[(j * 6 + i) * MatrixInputStride + InputOffset];
+			}
+		}
+
+		// Calculate A^TIA
+		WORK_TYPE ATI[4][6];
+		UNROLL
+		for (int i = 0; i < 6; i++)
+		{
+			ATI[0][i] = I[0][i] + I[1][i] + I[2][i] + I[3][i] + I[4][i];
+			ATI[1][i] = I[1][i] - I[2][i] + I[3][i] * 2 - I[4][i] * 2;
+			ATI[2][i] = I[1][i] + I[2][i] + I[3][i] * 4 + I[4][i] * 4;
+			ATI[3][i] = I[1][i] - I[2][i] + I[3][i] * 8 - I[4][i] * 8 + I[5][i];
+		}
+		WORK_TYPE ATIA[4][4];
+		UNROLL
+		for (int i = 0; i < 4; i++)
+		{
+			ATIA[i][0] = ATI[i][0] + ATI[i][1] + ATI[i][2] + ATI[i][3] + ATI[i][4];
+			ATIA[i][1] = ATI[i][1] - ATI[i][2] + ATI[i][3] * 2 - ATI[i][4] * 2;
+			ATIA[i][2] = ATI[i][1] + ATI[i][2] + ATI[i][3] * 4 + ATI[i][4] * 4;
+			ATIA[i][3] = ATI[i][1] - ATI[i][2] + ATI[i][3] * 8 - ATI[i][4] * 8 + ATI[i][5];
+		}
+
+#if HAS_BIAS
+		WORK_TYPE Scalar_BiasValue = Bias[Scalar_C];
+#endif
+
+		// Store result
+		for (int HOffset = 0; HOffset < 4; HOffset++)
+		{
+			const int HIndex = GroupID.y * 4 + HOffset;
+			if(HIndex < H)
+			{
+				for (int WOffset = 0; WOffset < 4; WOffset++)
+				{
+					const int WIndex = WBlock * 4 + WOffset;
+					if (WIndex < W)
+					{
+#if HAS_BIAS
+						Output[Scalar_OutputBase + HOffset * HOutputStride + WOffset] = ATIA[HOffset][WOffset] + Scalar_BiasValue;
+#else
+						Output[Scalar_OutputBase + HOffset * HOutputStride + WOffset] = ATIA[HOffset][WOffset];
+#endif
+					}
+				}
+			}
+		}
+	}
+}
diff --git a/Engine/Plugins/Experimental/NNERuntimeRDG/Shaders/Private/NNEHlslShaders/NNEHlslShadersConvWinogradWeights.usf b/Engine/Plugins/Experimental/NNERuntimeRDG/Shaders/Private/NNEHlslShaders/NNEHlslShadersConvWinogradWeights.usf
new file mode 100644
--- /dev/null
+++ b/Engine/Plugins/Experimental/NNERuntimeRDG/Shaders/Private/NNEHlslShaders/NNEHlslShadersConvWinogradWeights.usf
@@ -0,0 +1,78 @@
+// Copyright Epic Games, Inc. All Rights Reserved.
+
+#include "/Engine/Public/Platform.ush"
+
+#define KERNEL_STRIDE 9
+
+Buffer<float> Input; // Cw x Ci x 3 x 3
+RWBuffer<float> Output; // 36 x Ci x Cw
+
+int Ci;
+int Cw;
+
+int CwInputStride;
+
+int MatrixOutputStride;
+int CiOutputStride;
+
+
+[numthreads(THREADGROUP_SIZE_X, 1, 1)]
+void ConvWinogradWeights(
+	in const uint3 GroupID : SV_GroupID,
+	in const uint3 DispatchThreadID : SV_DispatchThreadID)
+{
+	const int ci = DispatchThreadID.x;
+	if (ci < Ci)
+	{
+		const int Scalar_Cw      = GroupID.y;
+		const int Input_Base 	= (Scalar_Cw * CwInputStride + ci) * KERNEL_STRIDE;
+		const int Output_Offset = ci * CiOutputStride + Scalar_Cw;
+
+		// Loading g
+		float g[3][3];
+		UNROLL
+		for (int j = 0; j < 3; j++)
+		{			
+			UNROLL
+			for (int i = 0; i < 3; i++)
+			{
+				g[j][i] = Input[Input_Base + j * 3 + i];
+			}
+		}
+
+		//Calculating GgG^T
+		float Gg[6][3];
+		UNROLL
+		for (int i = 0; i < 3; i++)
+		{
+			Gg[0][i] = g[0][i]/4;
+			Gg[1][i] = (-g[0][i] - g[1][i] - g[2][i]) / 6;
+			Gg[2][i] = (-g[0][i] + g[1][i] - g[2][i]) / 6;
+			Gg[3][i] = (g[0][i] + g[1][i]*2 + g[2][i]*4) / 24;
+			Gg[4][i] = (g[0][i] - g[1][i]*2 + g[2][i]*4) / 24;
+			Gg[5][i] = g[2][i];
+		}
+		float GgGT[6][6];
+		UNROLL
+		for (int i = 0; i < 6; i++)
+		{
+			GgGT[i][0] = Gg[i][0]/4;
+			GgGT[i][1] = (-Gg[i][0] - Gg[i][1] - Gg[i][2]) / 6;
+			GgGT[i][2] = (-Gg[i][0] + Gg[i][1] - Gg[i][2]) / 6;
+			GgGT[i][3] = (Gg[i][0] + Gg[i][1]*2 + Gg[i][2]*4) / 24;
+			GgGT[i][4] = (Gg[i][0] - Gg[i][1]*2 + Gg[i][2]*4) / 24;
+			GgGT[i][5] = Gg[i][2];
+		}
+
+		// Store result
+		UNROLL
+		for (int j = 0; j < 6; j++)
+		{
+			UNROLL
+			for (int i = 0; i < 6; i++)
+			{
+				Output[(j * 6 + i) * MatrixOutputStride + Output_Offset] = GgGT[j][i];
+			}
+		}
+	}
+}
diff --git a/Engine/Plugins/Experimental/NNERuntimeRDG/Source/NNEHlslShaders/Internal/NNEHlslShadersConvWinogradInputCS.h b/Engine/Plugins/Experimental/NNERuntimeRDG/Source/NNEHlslShaders/Internal/NNEHlslShadersConvWinogradInputCS.h
new file mode 100644
--- /dev/null
+++ b/Engine/Plugins/Experimental/NNERuntimeRDG/Source/NNEHlslShaders/Internal/NNEHlslShadersConvWinogradInputCS.h
@@ -0,0 +1,40 @@
+// Copyright Epic Games, Inc. All Rights Reserved.
+
+#pragma once
+
+#include "NNEHlslShaderBase.h"
+#include "RenderGraphUtils.h"
+
+namespace UE::NNEHlslShaders::Internal
+{
+	class FConvWinogradInputConstants
+	{
+	public:
+
+		static const int32 THREADGROUP_SIZE_X{ 32 };
+	};
+
+	class NNEHLSLSHADERS_API FConvWinogradInputCS : public FHlslShaderBase
+	{
+		DECLARE_GLOBAL_SHADER(FConvWinogradInputCS);
+		SHADER_USE_PARAMETER_STRUCT(FConvWinogradInputCS, FHlslShaderBase)
+
+	public:
+		BEGIN_SHADER_PARAMETER_STRUCT(FParameters, )
+			SHADER_PARAMETER_RDG_BUFFER_SRV(Buffer<float>, Input)
+			SHADER_PARAMETER_RDG_BUFFER_UAV(RWBuffer<float>, Output)
+			SHADER_PARAMETER(int32, C)
+			SHADER_PARAMETER(int32, H)
+			SHADER_PARAMETER(int32, W)
+			SHADER_PARAMETER(int32, WBlockCount)
+			SHADER_PARAMETER(int32, CInputStride)
+			SHADER_PARAMETER(int32, HInputStride)
+			SHADER_PARAMETER(int32, NiOutputStride)
+			SHADER_PARAMETER(int32, MatrixOutputStride)
+			SHADER_PARAMETER(int32, COutputStride)
+			SHADER_PARAMETER(int32, HOutputStride)
+		END_SHADER_PARAMETER_STRUCT()
+
+		static void ModifyCompilationEnvironment(const FGlobalShaderPermutationParameters& InParameters, FShaderCompilerEnvironment& OutEnvironment);
+	};
+} // UE::NNEHlslShaders::Internal
diff --git a/Engine/Plugins/Experimental/NNERuntimeRDG/Source/NNEHlslShaders/Internal/NNEHlslShadersConvWinogradMMMCS.h b/Engine/Plugins/Experimental/NNERuntimeRDG/Source/NNEHlslShaders/Internal/NNEHlslShadersConvWinogradMMMCS.h
new file mode 100644
--- /dev/null
+++ b/Engine/Plugins/Experimental/NNERuntimeRDG/Source/NNEHlslShaders/Internal/NNEHlslShadersConvWinogradMMMCS.h
@@ -0,0 +1,43 @@
+// Copyright Epic Games, Inc. All Rights Reserved.
+
+#pragma once
+
+#include "NNEHlslShaderBase.h"
+#include "NNEHlslShadersTypeHelper.h"
+#include "RenderGraphUtils.h"
+#include "ShaderParameterUtils.h"
+
+namespace UE::NNEHlslShaders::Internal
+{
+	class NNEHLSLSHADERS_API FConvWinogradMMMCS : public FHlslShaderBase
+	{
+		DECLARE_GLOBAL_SHADER(FConvWinogradMMMCS);
+		SHADER_USE_PARAMETER_STRUCT(FConvWinogradMMMCS, FHlslShaderBase)
+
+		class FDataType : SHADER_PERMUTATION_ENUM_CLASS("DATA_TYPE_ENUM", ENNEShaderDataType);
+		class FBlockSizeN : SHADER_PERMUTATION_SPARSE_INT("BLOCK_ELEM_COUNT_N", 16, 32, 64);
+		using FPermutationDomain = TShaderPermutationDomain<FDataType, FBlockSizeN>;
+
+	public:
+		BEGIN_SHADER_PARAMETER_STRUCT(FParameters, )
+			SHADER_PARAMETER_RDG_BUFFER_SRV(Buffer<float>, Input)
+			SHADER_PARAMETER_RDG_BUFFER_SRV(Buffer<float>, Weight)
+			SHADER_PARAMETER_RDG_BUFFER_UAV(RWBuffer<float>, Output)
+			SHADER_PARAMETER(int32, M)
+			SHADER_PARAMETER(int32, N)
+			SHADER_PARAMETER(int32, K)
+			SHADER_PARAMETER(int32, MatrixInputStride)
+			SHADER_PARAMETER(int32, KInputStride)
+			SHADER_PARAMETER(int32, MatrixWeightStride)
+			SHADER_PARAMETER(int32, KWeightStride)
+			SHADER_PARAMETER(int32, MatrixOutputStride)
+			SHADER_PARAMETER(int32, NOutputStride)
+		END_SHADER_PARAMETER_STRUCT()
+
+		static bool ShouldCompilePermutation(const FGlobalShaderPermutationParameters& Parameters);
+
+		static void ModifyCompilationEnvironment(const FGlobalShaderPermutationParameters& InParameters, FShaderCompilerEnvironment& OutEnvironment);
+
+		static int GetOptimalBlockSizeN(int M, int K, int N);
+	};
+} // UE::NNEHlslShaders::Internal
diff --git a/Engine/Plugins/Experimental/NNERuntimeRDG/Source/NNEHlslShaders/Internal/NNEHlslShadersConvWinogradOutputCS.h b/Engine/Plugins/Experimental/NNERuntimeRDG/Source/NNEHlslShaders/Internal/NNEHlslShadersConvWinogradOutputCS.h
new file mode 100644
--- /dev/null
+++ b/Engine/Plugins/Experimental/NNERuntimeRDG/Source/NNEHlslShaders/Internal/NNEHlslShadersConvWinogradOutputCS.h
@@ -0,0 +1,49 @@
+// Copyright Epic Games, Inc. All Rights Reserved.
+
+#pragma once
+
+#include "NNEHlslShaderBase.h"
+#include "NNEHlslShadersTypeHelper.h"
+#include "RenderGraphUtils.h"
+#include "ShaderParameterUtils.h"
+
+namespace UE::NNEHlslShaders::Internal
+{
+	class FConvWinogradOutputConstants
+	{
+	public:
+
+		static const int32 THREADGROUP_SIZE_X{ 32 };
+	};
+
+	class NNEHLSLSHADERS_API FConvWinogradOutputCS : public FHlslShaderBase
+	{
+		DECLARE_GLOBAL_SHADER(FConvWinogradOutputCS);
+		SHADER_USE_PARAMETER_STRUCT(FConvWinogradOutputCS, FHlslShaderBase)
+
+		class FHasBias : SHADER_PERMUTATION_BOOL("HAS_BIAS");
+		class FDataType : SHADER_PERMUTATION_ENUM_CLASS("DATA_TYPE_ENUM", ENNEShaderDataType);
+		using FPermutationDomain = TShaderPermutationDomain<FHasBias, FDataType>;
+
+	public:
+		BEGIN_SHADER_PARAMETER_STRUCT(FParameters, )
+			SHADER_PARAMETER_RDG_BUFFER_SRV(Buffer<float>, Input)
+			SHADER_PARAMETER_RDG_BUFFER_SRV(Buffer<float>, Bias)
+			SHADER_PARAMETER_RDG_BUFFER_UAV(RWBuffer<float>, Output)
+			SHADER_PARAMETER(int32, C)
+			SHADER_PARAMETER(int32, H)
+			SHADER_PARAMETER(int32, W)
+			SHADER_PARAMETER(int32, WBlockCount)
+			SHADER_PARAMETER(int32, NiInputStride)
+			SHADER_PARAMETER(int32, MatrixInputStride)
+			SHADER_PARAMETER(int32, CInputStride)
+			SHADER_PARAMETER(int32, HInputStride)
+			SHADER_PARAMETER(int32, COutputStride)
+			SHADER_PARAMETER(int32, HOutputStride)
+		END_SHADER_PARAMETER_STRUCT()
+
+		static bool ShouldCompilePermutation(const FGlobalShaderPermutationParameters& Parameters);
+
+		static void ModifyCompilationEnvironment(const FGlobalShaderPermutationParameters& InParameters, FShaderCompilerEnvironment& OutEnvironment);
+	};
+} // UE::NNEHlslShaders::Internal
diff --git a/Engine/Plugins/Experimental/NNERuntimeRDG/Source/NNEHlslShaders/Internal/NNEHlslShadersConvWinogradWeightsCS.h b/Engine/Plugins/Experimental/NNERuntimeRDG/Source/NNEHlslShaders/Internal/NNEHlslShadersConvWinogradWeightsCS.h
new file mode 100644
--- /dev/null
+++ b/Engine/Plugins/Experimental/NNERuntimeRDG/Source/NNEHlslShaders/Internal/NNEHlslShadersConvWinogradWeightsCS.h
@@ -0,0 +1,35 @@
+// Copyright Epic Games, Inc. All Rights Reserved.
+
+#pragma once
+
+#include "NNEHlslShaderBase.h"
+#include "RenderGraphUtils.h"
+
+namespace UE::NNEHlslShaders::Internal
+{
+	class FConvWinogradWeightsConstants
+	{
+	public:
+
+		static const int32 THREADGROUP_SIZE_X{ 32 };
+	};
+
+	class NNEHLSLSHADERS_API FConvWinogradWeightsCS : public FHlslShaderBase
+	{
+		DECLARE_GLOBAL_SHADER(FConvWinogradWeightsCS);
+		SHADER_USE_PARAMETER_STRUCT(FConvWinogradWeightsCS, FHlslShaderBase)
+
+	public:
+		BEGIN_SHADER_PARAMETER_STRUCT(FParameters, )
+			SHADER_PARAMETER_RDG_BUFFER_SRV(Buffer<float>, Input)
+			SHADER_PARAMETER_RDG_BUFFER_UAV(RWBuffer<float>, Output)
+			SHADER_PARAMETER(int32, Ci)
+			SHADER_PARAMETER(int32, Cw)
+			SHADER_PARAMETER(int32, CwInputStride)
+			SHADER_PARAMETER(int32, MatrixOutputStride)
+			SHADER_PARAMETER(int32, CiOutputStride)
+		END_SHADER_PARAMETER_STRUCT()
+
+		static void ModifyCompilationEnvironment(const FGlobalShaderPermutationParameters& InParameters, FShaderCompilerEnvironment& OutEnvironment);
+	};
+} // UE::NNEHlslShaders::Internal
diff --git a/Engine/Plugins/Experimental/NNERuntimeRDG/Source/NNEHlslShaders/Private/NNEHlslShadersConvWinogradInputCS.cpp b/Engine/Plugins/Experimental/NNERuntimeRDG/Source/NNEHlslShaders/Private/NNEHlslShadersConvWinogradInputCS.cpp
new file mode 100644
--- /dev/null
+++ b/Engine/Plugins/Experimental/NNERuntimeRDG/Source/NNEHlslShaders/Private/NNEHlslShadersConvWinogradInputCS.cpp
@@ -0,0 +1,14 @@
+// Copyright Epic Games, Inc. All Rights Reserved.
+
+#include "NNEHlslShadersConvWinogradInputCS.h"
+#include "NNE.h"
+
+namespace UE::NNEHlslShaders::Internal
+{
+	void FConvWinogradInputCS::ModifyCompilationEnvironment(const FGlobalShaderPermutationParameters& InParameters, FShaderCompilerEnvironment& OutEnvironment)
+	{
+		OutEnvironment.SetDefine(TEXT("THREADGROUP_SIZE_X"), FConvWinogradInputConstants::THREADGROUP_SIZE_X);
+	}
+
+	IMPLEMENT_GLOBAL_SHADER(FConvWinogradInputCS, "/NNEHlslShaders/NNEHlslShadersConvWinogradInput.usf", "ConvWinogradInput", SF_Compute);
+} // UE::NNEHlslShaders::Internal
\ No newline at end of file
diff --git a/Engine/Plugins/Experimental/NNERuntimeRDG/Source/NNEHlslShaders/Private/NNEHlslShadersConvWinogradMMMCS.cpp b/Engine/Plugins/Experimental/NNERuntimeRDG/Source/NNEHlslShaders/Private/NNEHlslShadersConvWinogradMMMCS.cpp
new file mode 100644
--- /dev/null
+++ b/Engine/Plugins/Experimental/NNERuntimeRDG/Source/NNEHlslShaders/Private/NNEHlslShadersConvWinogradMMMCS.cpp
@@ -0,0 +1,71 @@
+// Copyright Epic Games, Inc. All Rights Reserved.
+
+#include "NNEHlslShadersConvWinogradMMMCS.h"
+
+#include <limits>
+#include "NNEHlslShadersTypeHelper.h"
+#include "NNE.h"
+#include "ShaderCompilerCore.h"
+
+namespace UE::NNEHlslShaders::Internal
+{
+	bool FConvWinogradMMMCS::ShouldCompilePermutation(const FGlobalShaderPermutationParameters& Parameters)
+	{
+		if (!FHlslShaderBase::ShouldCompilePermutation(Parameters))
+		{
+			return false;
+		}
+		FPermutationDomain PermutationVector(Parameters.PermutationId);
+		ENNEShaderDataType DataType = PermutationVector.Get<FConvWinogradMMMCS::FDataType>();
+		return DataType == ENNEShaderDataType::FLOAT16 || DataType == ENNEShaderDataType::FLOAT32;
+	}
+
+	void FConvWinogradMMMCS::ModifyCompilationEnvironment(const FGlobalShaderPermutationParameters& InParameters, FShaderCompilerEnvironment& OutEnvironment)
+	{
+		FPermutationDomain PermutationVector(InParameters.PermutationId);
+		ENNEShaderDataType DataType = PermutationVector.Get<FConvWinogradMMMCS::FDataType>();
+		OutEnvironment.SetDefine(TEXT("WORK_TYPE"), ShaderDataTypeToName(DataType));
+
+		FGlobalShader::ModifyCompilationEnvironment(InParameters, OutEnvironment);
+		OutEnvironment.CompilerFlags.Add(CFLAG_AllowRealTypes);
+	}
+
+	// Simple Heuristic that tries to find the optimal block size
+	int FConvWinogradMMMCS::GetOptimalBlockSizeN(int M, int K, int N)
+	{
+		struct PerformanceData
+		{
+			int BlockSizeN;
+			float LowFlops;
+			float HighFlops;
+		};
+
+		// These numbers are based on performance measurements
+		PerformanceData Performance[3] =
+		{
+			{ 16, 125, 250 },
+			{ 32, 200, 280 },
+			{ 64, 290, 300 }
+		};
+
+		bool UseHighFlops = M / N < 250;
+		float SmallestDuration = std::numeric_limits<float>::infinity();
+		int Result = 0;
+		for (int i = 0; i < sizeof(Performance) / sizeof(Performance[0]); i++)
+		{
+			float Flops = UseHighFlops ? Performance[i].HighFlops : Performance[i].LowFlops;
+			int Mask = Performance[i].BlockSizeN - 1;
+			// Round up to the next multiple of BlocksSizeN
+			int NExtended = (N + Mask) & ~Mask;
+			float Duration = NExtended / Flops;
+			if (SmallestDuration > Duration)
+			{
+				SmallestDuration = Duration;
+				Result = Performance[i].BlockSizeN;
+			}
+		}
+		return Result;
+	}
+
+	IMPLEMENT_GLOBAL_SHADER(FConvWinogradMMMCS, "/NNEHlslShaders/NNEHlslShadersConvWinogradMMM.usf", "ConvWinogradMMM", SF_Compute);
+} // UE::NNEHlslShaders::Internal
diff --git a/Engine/Plugins/Experimental/NNERuntimeRDG/Source/NNEHlslShaders/Private/NNEHlslShadersConvWinogradOutputCS.cpp b/Engine/Plugins/Experimental/NNERuntimeRDG/Source/NNEHlslShaders/Private/NNEHlslShadersConvWinogradOutputCS.cpp
new file mode 100644
--- /dev/null
+++ b/Engine/Plugins/Experimental/NNERuntimeRDG/Source/NNEHlslShaders/Private/NNEHlslShadersConvWinogradOutputCS.cpp
@@ -0,0 +1,34 @@
+// Copyright Epic Games, Inc. All Rights Reserved.
+
+#include "NNEHlslShadersConvWinogradOutputCS.h"
+
+#include "NNEHlslShadersTypeHelper.h"
+#include "NNE.h"
+#include "ShaderCompilerCore.h"
+
+namespace UE::NNEHlslShaders::Internal
+{
+	bool FConvWinogradOutputCS::ShouldCompilePermutation(const FGlobalShaderPermutationParameters& Parameters)
+	{
+		if (!FHlslShaderBase::ShouldCompilePermutation(Parameters))
+		{
+			return false;
+		}
+		FPermutationDomain PermutationVector(Parameters.PermutationId);
+		ENNEShaderDataType DataType = PermutationVector.Get<FConvWinogradOutputCS::FDataType>();
+		return DataType == ENNEShaderDataType::FLOAT16 || DataType == ENNEShaderDataType::FLOAT32;
+	}
+
+	void FConvWinogradOutputCS::ModifyCompilationEnvironment(const FGlobalShaderPermutationParameters& InParameters, FShaderCompilerEnvironment& OutEnvironment)
+	{
+		FPermutationDomain PermutationVector(InParameters.PermutationId);
+		ENNEShaderDataType DataType = PermutationVector.Get<FConvWinogradOutputCS::FDataType>();
+		OutEnvironment.SetDefine(TEXT("WORK_TYPE"), ShaderDataTypeToName(DataType));
+
+		OutEnvironment.SetDefine(TEXT("THREADGROUP_SIZE_X"), FConvWinogradOutputConstants::THREADGROUP_SIZE_X);
+		FGlobalShader::ModifyCompilationEnvironment(InParameters, OutEnvironment);
+		OutEnvironment.CompilerFlags.Add(CFLAG_AllowRealTypes);
+	}
+
+	IMPLEMENT_GLOBAL_SHADER(FConvWinogradOutputCS, "/NNEHlslShaders/NNEHlslShadersConvWinogradOutput.usf", "ConvWinogradOutput", SF_Compute);
+} // UE::NNEHlslShaders::Internal
diff --git a/Engine/Plugins/Experimental/NNERuntimeRDG/Source/NNEHlslShaders/Private/NNEHlslShadersConvWinogradWeightsCS.cpp b/Engine/Plugins/Experimental/NNERuntimeRDG/Source/NNEHlslShaders/Private/NNEHlslShadersConvWinogradWeightsCS.cpp
new file mode 100644
--- /dev/null
+++ b/Engine/Plugins/Experimental/NNERuntimeRDG/Source/NNEHlslShaders/Private/NNEHlslShadersConvWinogradWeightsCS.cpp
@@ -0,0 +1,14 @@
+// Copyright Epic Games, Inc. All Rights Reserved.
+
+#include "NNEHlslShadersConvWinogradWeightsCS.h"
+#include "NNE.h"
+
+namespace UE::NNEHlslShaders::Internal
+{
+	void FConvWinogradWeightsCS::ModifyCompilationEnvironment(const FGlobalShaderPermutationParameters& InParameters, FShaderCompilerEnvironment& OutEnvironment)
+	{
+		OutEnvironment.SetDefine(TEXT("THREADGROUP_SIZE_X"), FConvWinogradWeightsConstants::THREADGROUP_SIZE_X);
+	}
+
+	IMPLEMENT_GLOBAL_SHADER(FConvWinogradWeightsCS, "/NNEHlslShaders/NNEHlslShadersConvWinogradWeights.usf", "ConvWinogradWeights", SF_Compute);
+} // UE::NNEHlslShaders::Internal
\ No newline at end of file
diff --git a/Engine/Plugins/Experimental/NNERuntimeRDG/Source/NNERuntimeRDG/Private/Hlsl/NNERuntimeRDGConv.cpp b/Engine/Plugins/Experimental/NNERuntimeRDG/Source/NNERuntimeRDG/Private/Hlsl/NNERuntimeRDGConv.cpp
--- a/Engine/Plugins/Experimental/NNERuntimeRDG/Source/NNERuntimeRDG/Private/Hlsl/NNERuntimeRDGConv.cpp
+++ b/Engine/Plugins/Experimental/NNERuntimeRDG/Source/NNERuntimeRDG/Private/Hlsl/NNERuntimeRDGConv.cpp
@@ -1,12 +1,18 @@
 // Copyright Epic Games, Inc. All Rights Reserved.
 
 #include "NNERuntimeRDGConv.h"
+
 #include "Algo/AnyOf.h"
+#include "HAL/IConsoleManager.h"
 #include "NNEAttributeMap.h"
 #include "NNETensor.h"
 #include "NNETypes.h"
 #include "NNEHlslShadersConvCS.h"
 #include "NNEHlslShadersConvMatmulCS.h"
+#include "NNEHlslShadersConvWinogradInputCS.h"
+#include "NNEHlslShadersConvWinogradMMMCS.h"
+#include "NNEHlslShadersConvWinogradOutputCS.h"
+#include "NNEHlslShadersConvWinogradWeightsCS.h"
 #include "NNEHlslShadersTypeHelper.h"
 #include "NNERuntimeRDGHelperTranspose.h"
 #include "NNERuntimeRDGHlslHelper.h"
@@ -15,10 +21,24 @@ namespace UE::NNERuntimeRDG::Private::Hlsl
 {
 	DECLARE_GPU_STAT_NAMED(FNNEOperatorConvDefault, TEXT("NNE.Operator.Hlsl.Conv.Default"));
 	DECLARE_GPU_STAT_NAMED(FNNEOperatorConvMatmul, TEXT("NNE.Operator.Hlsl.Conv.Matmul"));
+	DECLARE_GPU_STAT_NAMED(FNNEOperatorConvWinogradInput, TEXT("NNE.Operator.Hlsl.Conv.Winograd.Input"));
+	DECLARE_GPU_STAT_NAMED(FNNEOperatorConvWinogradMMM, TEXT("NNE.Operator.Hlsl.Conv.Winograd.MMM"));
+	DECLARE_GPU_STAT_NAMED(FNNEOperatorConvWinogradOutput, TEXT("NNE.Operator.Hlsl.Conv.Winograd.Output"));
+	DECLARE_GPU_STAT_NAMED(FNNEOperatorConvWinogradWeights, TEXT("NNE.Operator.Hlsl.Conv.Winograd.Weights"));
 
 	using EConvAutoPad = UE::NNEHlslShaders::Internal::EConvAutoPad;
 	using EConvGroupSize = UE::NNEHlslShaders::Internal::EConvGroupSize;
 
+
+	TAutoConsoleVariable<int32> CVarWinogradPrecision(
+		TEXT("nne.hlsl.WinogradPrecision"),
+		0,
+		TEXT("Selects the Precision of the Winograd Convolution implementation.\n")
+		TEXT(" 0: Allow Float16 (fast, low precision) (default)\n")
+		TEXT(" 1: Disable Float16 (medium speed, medium precision)\n")
+		TEXT(" 2: Disable Winograd (slow)")
+	);
+
 	/**
 	 * Convolution operator implementation
 	 */
@@ -26,6 +46,19 @@ namespace UE::NNERuntimeRDG::Private::Hlsl
 	{
 	public:
 
+		// The first values have to correspond to the values in CVarWinogradPrecision
+		// so that it is possible to convert from the CVar value to this enum.
+		enum class EWinogradPrecision
+		{
+			FP16 = 0,
+			FP32,
+			DISABLED,
+
+			UNDEFINED,
+
+			MAX
+		};
+
 		static FOperatorHlsl* Create()
 		{
 			return new FConv();
@@ -48,6 +81,7 @@ namespace UE::NNERuntimeRDG::Private::Hlsl
 		bool bAreWeightsTransposed = false;
 		EPixelFormat BufferPixelFormat;
 		EConvGroupSize GroupSize = EConvGroupSize::MAX;
+		EWinogradPrecision WinogradTestPrecision;
 
 	public:
 
@@ -124,12 +158,21 @@ namespace UE::NNERuntimeRDG::Private::Hlsl
 				Pads = Attributes.GetValueOrDefault<TArray<int32>>(TEXT("pads"), PadsDefault);
 			}
 			Strides = Attributes.GetValueOrDefault<TArray<int32>>(TEXT("strides"), DilationsOrStridesDefault);
+			if (Strides.Num() != NumDimensions)
+			{
+				UE_LOG(LogNNE, Warning, TEXT("Conv strides count doesn't match the number of spatial dimensions"));
+				return false;
+			}
 			BufferPixelFormat = TensorDataTypeToPixelFormat(Input.GetDataType());
 
+			int32 WinogradTestPrecisionValue = Attributes.GetValueOrDefault<int32>(TEXT("__UE__WinogradPrecision"), static_cast<int32>(EWinogradPrecision::UNDEFINED));
+			WinogradTestPrecisionValue = FMath::Clamp(WinogradTestPrecisionValue, 0, static_cast<int32>(EWinogradPrecision::MAX) - 1);
+			WinogradTestPrecision = static_cast<EWinogradPrecision>(WinogradTestPrecisionValue);
+
 			return true;
 		}
 
-		void DispatchConvDefault (FRDGBuilder& GraphBuilder, const FTensorRDG& Input, const FTensorRDG& Weights, const FTensorRDG* Bias, const FTensorRDG& Output)
+		void DispatchConvDefault (FRDGBuilder& GraphBuilder, const FTensorRDG& Input, const FTensorRDG& Weights, const FTensorRDG* Bias, const FTensorRDG& Output) const
 		{
 			using namespace UE::NNEHlslShaders::Internal;
 
@@ -170,7 +213,7 @@ namespace UE::NNERuntimeRDG::Private::Hlsl
 				FConvCS::GetGroupCount(OutputShape, FConvCS::GetGroupShape(GroupSize, NumDimensions)));
 		}
 
-		bool DispatchConvMatmul(FRDGBuilder& GraphBuilder, const FTensorRDG& Input, const FTensorRDG& Weights, const FTensorRDG* Bias, const FTensorRDG& Output)
+		bool DispatchConvMatmul(FRDGBuilder& GraphBuilder, const FTensorRDG& Input, const FTensorRDG& Weights, const FTensorRDG* Bias, const FTensorRDG& Output) const
 		{
 			using namespace UE::NNEHlslShaders::Internal;
 
@@ -246,6 +289,227 @@ namespace UE::NNERuntimeRDG::Private::Hlsl
 			return true;
 		}
 
+		// This algorithm is inspired by this paper: [REDACTED_PROFILE_URL]/CUDA-Winograd/blob/master/report.pdf
+		bool DispatchConvWinograd(FRDGBuilder& GraphBuilder, const FTensorRDG& Input, const FTensorRDG& Weights, const FTensorRDG* Bias, const FTensorRDG& Output) const
+		{
+			using namespace UE::NNEHlslShaders::Internal;
+
+			if (Group != 1)
+				return false;
+			if (Input.GetShape().Rank() != 4)
+				return false;
+			if (Output.GetShape().Rank() != 4)
+				return false;
+			if (Weights.GetShape().Rank() != 4)
+				return false;
+
+			check(Strides.Num() == 2);
+			if (Strides[0] != 1 || Strides[1] != 1)
+				return false;
+			
+			if (bHasDilation)
+				return false;
+
+			TArray<int32> Padding = FConvCS::GetPadding(Input.GetShape().GetData(), Weights.GetShape().GetData(), AutoPad, Dilations, Strides, Pads);
+
+			check(Padding.Num() == 4);
+			if (Padding[0] != 1 || Padding[1] != 1 || Padding[2] != 1 || Padding[3] != 1)
+				return false;
+
+			const int Hw = Weights.GetShape().GetData()[2];
+			const int Ww = Weights.GetShape().GetData()[3];
+			if (Hw != 3 || Ww != 3)
+				return false;
+
+			check(!bAreWeightsTransposed);
+
+			const int Ni = Input.GetShape().GetData()[0];
+			const int Ci = Input.GetShape().GetData()[1];
+			const int Hi = Input.GetShape().GetData()[2];
+			const int Wi = Input.GetShape().GetData()[3];
+
+			check(Ni == Output.GetShape().GetData()[0]);
+			const int Cw = Output.GetShape().GetData()[1];
+			check(Hi == Output.GetShape().GetData()[2]);
+			check(Wi == Output.GetShape().GetData()[3]);
+
+			check(Cw == Weights.GetShape().GetData()[0]);
+			check(Ci == Weights.GetShape().GetData()[1]);
+
+			EWinogradPrecision Precision = WinogradTestPrecision;
+			if (Precision == EWinogradPrecision::UNDEFINED)
+			{
+				int32 CVarValue = CVarWinogradPrecision.GetValueOnRenderThread();
+				CVarValue = FMath::Clamp(CVarValue, 0, static_cast<int32>(EWinogradPrecision::UNDEFINED) - 1);
+				Precision = static_cast<EWinogradPrecision>(CVarValue);
+			}
+
+			if (Precision == EWinogradPrecision::DISABLED)
+				return false;
+			
+			check(Input.GetDataType() == ENNETensorDataType::Float || Input.GetDataType() == ENNETensorDataType::Half);
+			const ENNETensorDataType TensorDataType = Precision == EWinogradPrecision::FP16 ? Input.GetDataType() : ENNETensorDataType::Float;
+
+
+			const int BlockCountH = FMath::DivideAndRoundUp(Hi, 4);
+			const int BlockCountW = FMath::DivideAndRoundUp(Wi, 4);
+
+			//Ensures that the MMM shader gets called with an even number of elements in the M dimension
+			const int BlockCountWExtended = ((BlockCountH * BlockCountW) % 2) != 0 ? RoundUp<2>(BlockCountW) : BlockCountW;
+			const int CwExtended = RoundUp<2>(Cw);
+
+			const int TransformedWeightSize = 36 * CwExtended * Ci;
+			const int TransformedInputSize = Ni * 36 * (BlockCountH * BlockCountWExtended) * Ci;
+			const int TransformedOutputSize = Ni * 36 * (BlockCountH * BlockCountWExtended) * CwExtended;
+
+			const ENNEShaderDataType ShaderDataType = TensorToShaderDataType(TensorDataType);
+			const EPixelFormat IntermediateBufferPixelFormat = TensorDataTypeToPixelFormat(TensorDataType);
+			const EPixelFormat IntermediateBufferPixelFormat2 = TensorDataType == ENNETensorDataType::Float ? PF_G32R32F : PF_G16R16F;
+			const uint32 IntermediateElementByteSize = UE::NNE::GetTensorDataTypeSizeInBytes(TensorDataType);
+
+			const FRDGBufferDesc TransformedWeightsBufferDesc = FRDGBufferDesc::CreateBufferDesc(IntermediateElementByteSize, TransformedWeightSize);
+			const FRDGBufferDesc TransformedInputBufferDesc = FRDGBufferDesc::CreateBufferDesc(IntermediateElementByteSize, TransformedInputSize);
+			const FRDGBufferDesc TransformedOutputBufferDesc = FRDGBufferDesc::CreateBufferDesc(IntermediateElementByteSize, TransformedOutputSize);
+
+			const FRDGBufferRef TransformedWeights = GraphBuilder.CreateBuffer(TransformedWeightsBufferDesc, TEXT("NNE.Tensor.ConvWinograd.TransformedWeights"), ERDGBufferFlags::None);
+			const FRDGBufferRef TransformedInput = GraphBuilder.CreateBuffer(TransformedInputBufferDesc, TEXT("NNE.Tensor.ConvWinograd.TransformedInput"), ERDGBufferFlags::None);
+			const FRDGBufferRef TransformedOutput = GraphBuilder.CreateBuffer(TransformedOutputBufferDesc, TEXT("NNE.Tensor.ConvWinograd.TransformedOutput"), ERDGBufferFlags::None);
+
+
+			// Dispatch Weight transformation
+			{
+				FConvWinogradWeightsCS::FParameters* Params = GraphBuilder.AllocParameters<FConvWinogradWeightsCS::FParameters>();
+				Params->Input = GraphBuilder.CreateSRV(FRDGBufferSRVDesc(Weights.GetBuffer(), BufferPixelFormat));
+				Params->Output = GraphBuilder.CreateUAV(FRDGBufferUAVDesc(TransformedWeights, IntermediateBufferPixelFormat));
+				Params->Ci = Ci;
+				Params->Cw = Cw;
+				Params->CwInputStride = Ci;
+				Params->MatrixOutputStride = Ci * CwExtended;
+				Params->CiOutputStride = CwExtended;
+
+				TShaderMapRef<FConvWinogradWeightsCS> ComputeShader(GetGlobalShaderMap(GMaxRHIFeatureLevel));
+
+				RDG_EVENT_SCOPE(GraphBuilder, "NNE.Operator.Hlsl.Conv.Winograd.Weights");
+				RDG_GPU_STAT_SCOPE(GraphBuilder, FNNEOperatorConvWinogradWeights);
+
+				FComputeShaderUtils::AddPass(
+					GraphBuilder,
+					RDG_EVENT_NAME("NNE.Operator.Hlsl.Conv.Winograd.Weights.Dispatch"),
+					ERDGPassFlags::Compute | ERDGPassFlags::NeverCull,
+					ComputeShader,
+					Params,
+					FIntVector(FMath::DivideAndRoundUp(Ci, FConvWinogradWeightsConstants::THREADGROUP_SIZE_X), CwExtended, 1));
+			}
+			// Dispatch Input transformation
+			{
+				FConvWinogradInputCS::FParameters* Params = GraphBuilder.AllocParameters<FConvWinogradInputCS::FParameters>();
+				Params->Input = GraphBuilder.CreateSRV(FRDGBufferSRVDesc(Input.GetBuffer(), BufferPixelFormat));
+				Params->Output = GraphBuilder.CreateUAV(FRDGBufferUAVDesc(TransformedInput, IntermediateBufferPixelFormat));
+				Params->C = Ci;
+				Params->H = Hi;
+				Params->W = Wi;
+				Params->WBlockCount = BlockCountW;
+				Params->CInputStride = Hi * Wi;
+				Params->HInputStride = Wi;
+				Params->NiOutputStride = 36 * Ci * BlockCountH * BlockCountWExtended;
+				Params->MatrixOutputStride = Ci * BlockCountH * BlockCountWExtended;
+				Params->COutputStride = BlockCountH * BlockCountWExtended;
+				Params->HOutputStride = BlockCountWExtended;
+
+				TShaderMapRef<FConvWinogradInputCS> ComputeShader(GetGlobalShaderMap(GMaxRHIFeatureLevel));
+
+				RDG_EVENT_SCOPE(GraphBuilder, "NNE.Operator.Hlsl.Conv.Winograd.Input");
+				RDG_GPU_STAT_SCOPE(GraphBuilder, FNNEOperatorConvWinogradInput);
+
+				FComputeShaderUtils::AddPass(
+					GraphBuilder,
+					RDG_EVENT_NAME("NNE.Operator.Hlsl.Conv.Winograd.Input.Dispatch"),
+					ERDGPassFlags::Compute | ERDGPassFlags::NeverCull,
+					ComputeShader,
+					Params,
+					FIntVector(FMath::DivideAndRoundUp(BlockCountWExtended, FConvWinogradInputConstants::THREADGROUP_SIZE_X), BlockCountH, Ci * Ni));
+			}
+			// Dispatch MMM
+			{
+				const int VectorSize = 2;
+				const int M = (BlockCountH * BlockCountWExtended);
+				const int N = CwExtended;
+				const int K = Ci;
+				FConvWinogradMMMCS::FParameters* Params = GraphBuilder.AllocParameters<FConvWinogradMMMCS::FParameters>();
+				Params->Input = GraphBuilder.CreateSRV(FRDGBufferSRVDesc(TransformedInput, IntermediateBufferPixelFormat2));
+				Params->Weight = GraphBuilder.CreateSRV(FRDGBufferSRVDesc(TransformedWeights, IntermediateBufferPixelFormat2));
+				Params->Output = GraphBuilder.CreateUAV(FRDGBufferUAVDesc(TransformedOutput, IntermediateBufferPixelFormat2));
+				Params->M = M / VectorSize;
+				Params->N = N / VectorSize;
+				Params->K = K;
+				Params->MatrixInputStride = K * M / VectorSize;
+				Params->KInputStride = M / VectorSize;
+				Params->MatrixWeightStride = K * N / VectorSize;
+				Params->KWeightStride = N / VectorSize;
+				Params->MatrixOutputStride = N * M / VectorSize;
+				Params->NOutputStride = M / VectorSize;
+
+				const int BlockSizeN = FConvWinogradMMMCS::GetOptimalBlockSizeN(M, K, N);
+				FConvWinogradMMMCS::FPermutationDomain PermutationVector;
+				PermutationVector.Set<FConvWinogradMMMCS::FDataType>(ShaderDataType);
+				PermutationVector.Set<FConvWinogradMMMCS::FBlockSizeN>(BlockSizeN);
+				TShaderMapRef<FConvWinogradMMMCS> ComputeShader(GetGlobalShaderMap(GMaxRHIFeatureLevel), PermutationVector);
+
+				RDG_EVENT_SCOPE(GraphBuilder, "NNE.Operator.Hlsl.Conv.Winograd.MMM");
+				RDG_GPU_STAT_SCOPE(GraphBuilder, FNNEOperatorConvWinogradMMM);
+
+				FComputeShaderUtils::AddPass(
+					GraphBuilder,
+					RDG_EVENT_NAME("NNE.Operator.Hlsl.Conv.Winograd.MMM.Dispatch"),
+					ERDGPassFlags::Compute | ERDGPassFlags::NeverCull,
+					ComputeShader,
+					Params,
+					FIntVector(FMath::DivideAndRoundUp(M, 64), FMath::DivideAndRoundUp(N, BlockSizeN), 36 * Ni));
+			}
+			// Dispatch Output transformation
+			{
+				const int C = Cw;
+				const int H = Hi;
+				const int W = Wi;
+				const bool bHasBias = Bias != nullptr;
+				FConvWinogradOutputCS::FParameters* Params = GraphBuilder.AllocParameters<FConvWinogradOutputCS::FParameters>();
+				Params->Input = GraphBuilder.CreateSRV(FRDGBufferSRVDesc(TransformedOutput, IntermediateBufferPixelFormat));
+				if (bHasBias)
+				{
+					Params->Bias = GraphBuilder.CreateSRV(Bias->GetBuffer(), BufferPixelFormat);
+				}
+				Params->Output = GraphBuilder.CreateUAV(FRDGBufferUAVDesc(Output.GetBuffer(), BufferPixelFormat));
+				Params->C = C;
+				Params->H = H;
+				Params->W = W;
+				Params->WBlockCount = BlockCountW;
+				Params->NiInputStride = 36 * CwExtended * BlockCountH * BlockCountWExtended;
+				Params->MatrixInputStride = CwExtended * BlockCountH * BlockCountWExtended;
+				Params->CInputStride = BlockCountH * BlockCountWExtended;
+				Params->HInputStride = BlockCountWExtended;
+				Params->COutputStride = H * W;
+				Params->HOutputStride = W;
+
+				FConvWinogradOutputCS::FPermutationDomain PermutationVector;
+				PermutationVector.Set<FConvWinogradOutputCS::FHasBias>(bHasBias);
+				PermutationVector.Set<FConvWinogradOutputCS::FDataType>(ShaderDataType);
+				TShaderMapRef<FConvWinogradOutputCS> ComputeShader(GetGlobalShaderMap(GMaxRHIFeatureLevel), PermutationVector);
+
+				RDG_EVENT_SCOPE(GraphBuilder, "NNE.Operator.Hlsl.Conv.Winograd.Output");
+				RDG_GPU_STAT_SCOPE(GraphBuilder, FNNEOperatorConvWinogradOutput);
+
+				FComputeShaderUtils::AddPass(
+					GraphBuilder,
+					RDG_EVENT_NAME("NNE.Operator.Hlsl.Conv.Winograd.Output.Dispatch"),
+					ERDGPassFlags::Compute | ERDGPassFlags::NeverCull,
+					ComputeShader,
+					Params,
+					FIntVector(FMath::DivideAndRoundUp(BlockCountW, FConvWinogradOutputConstants::THREADGROUP_SIZE_X), BlockCountH, Cw * Ni));
+			}
+
+			return true;
+		}
+
 		virtual void Dispatch(FRDGBuilder& GraphBuilder, TConstArrayView<FTensorRDGRef> InputTensors, TConstArrayView<FTensorRDGRef> OutputTensors) override
 		{
 			check(InputTensors.Num() >= 2 && InputTensors.Num() <= 3);
@@ -269,6 +533,10 @@ namespace UE::NNERuntimeRDG::Private::Hlsl
 			check(Output.GetShape().Rank() == Input.GetShape().Rank());
 			check(NumDimensions == (Input.GetShape().Rank() - 2));
 
+			if (DispatchConvWinograd(GraphBuilder, Input, Weights, Bias, Output))
+			{
+				return;
+			}
 			if (DispatchConvMatmul(GraphBuilder, Input, Weights, Bias, Output))
 			{
 				return;
@@ -293,12 +561,27 @@ namespace UE::NNERuntimeRDG::Private::Hlsl
 			if (bHasDilation)
 				return;
 
+			// Don't transpose if Winograd might be possible
+			TArrayView<const uint32> WeightsShape = Weights->GetShape().GetData();
+			bool IsKernel3x3 = WeightsShape[2] == 3 && WeightsShape[3] == 3;
+			check(Strides.Num() == 2);
+			bool AreStrides1 = Strides[0] == 1 && Strides[1] == 1;
+			if (IsKernel3x3 && AreStrides1)
+				return;
+
 			//Transpose from CwCiHwWw to HwWwCiCw
 			if (Internal::CPUHelper::Transpose::TransposePreparedData(*Weights, {2,3,1,0} ))
 			{
 				bAreWeightsTransposed = true;
 			}
 		};
+
+	private:
+		template<int divider>
+		static int RoundUp(int value)
+		{
+			return FMath::DivideAndRoundUp(value, divider) * divider;
+		}
 	};
 
 	bool ValidateConvOperator(const NNE::FAttributeMap& AttributeMap, TConstArrayView<ENNETensorDataType> InputTypes, TConstArrayView<NNE::FSymbolicTensorShape> InputShapes)
@@ -312,6 +595,7 @@ namespace UE::NNERuntimeRDG::Private::Hlsl
 		AttributeValidator.AddOptional(TEXT("kernel_shape"), ENNEAttributeDataType::Int32Array); // idea: cross check input weight shape with this attribute if present
 		AttributeValidator.AddOptional(TEXT("pads"), ENNEAttributeDataType::Int32Array);
 		AttributeValidator.AddOptional(TEXT("strides"), ENNEAttributeDataType::Int32Array);
+		AttributeValidator.AddOptional(TEXT("__UE__WinogradPrecision"), ENNEAttributeDataType::Int32);
 
 		bIsValid &= AttributeValidator.Validate(AttributeMap);
 
