[NNE] NNERuntimeRDGHlsl,NNEQA/Incubator log cleanup

Problem:
* Log are inconsistently prefixed by class/function, also NNERuntimeRDGHlsl and NNEQA didn't had there own log category.

Solution:
* Unifiy log by ensuring there is no prefixes but for the operator name where relevant in the hlsl runtime.
* Add a log category to the hlsl runtime and to the QA plugin/project.

Bonus:
* Cleanup some of dependancies of NNEIncubator to RDGHlsl

Tests:
* Win64 Editor debug : smoke tests + style transfer
* Win64 Standalone packaged debug : style transfer

#rnx

diff --git a/Engine/Plugins/Experimental/NNERuntimeRDG/Source/NNEHlslShaders/Internal/NNEHlslShaderBase.h b/Engine/Plugins/Experimental/NNERuntimeRDG/Source/NNEHlslShaders/Internal/NNEHlslShaderBase.h
deleted file mode 100644
diff --git a/Engine/Plugins/Experimental/NNERuntimeRDG/Source/NNEHlslShaders/Internal/NNEHlslShadersBase.h b/Engine/Plugins/Experimental/NNERuntimeRDG/Source/NNEHlslShaders/Internal/NNEHlslShadersBase.h
new file mode 100644
--- /dev/null
+++ b/Engine/Plugins/Experimental/NNERuntimeRDG/Source/NNEHlslShaders/Internal/NNEHlslShadersBase.h
@@ -0,0 +1,19 @@
+// Copyright Epic Games, Inc. All Rights Reserved.
+
+#pragma once
+
+#include "GlobalShader.h"
+
+namespace UE::NNEHlslShaders::Internal
+{
+	class FHlslShaderBase : public FGlobalShader
+	{
+	public:
+
+		static bool ShouldCompilePermutation(const FGlobalShaderPermutationParameters& Parameters);
+
+		FHlslShaderBase() {}
+		FHlslShaderBase(const ShaderMetaType::CompiledShaderInitializerType& initType) : FGlobalShader(initType) {}
+	};
+	
+} // UE::NNEHlslShaders::Internal
\ No newline at end of file
diff --git a/Engine/Plugins/Experimental/NNERuntimeRDG/Source/NNEHlslShaders/Internal/NNEHlslShadersBatchNormalizationCS.h b/Engine/Plugins/Experimental/NNERuntimeRDG/Source/NNEHlslShaders/Internal/NNEHlslShadersBatchNormalizationCS.h
--- a/Engine/Plugins/Experimental/NNERuntimeRDG/Source/NNEHlslShaders/Internal/NNEHlslShadersBatchNormalizationCS.h
+++ b/Engine/Plugins/Experimental/NNERuntimeRDG/Source/NNEHlslShaders/Internal/NNEHlslShadersBatchNormalizationCS.h
@@ -2,8 +2,8 @@
 
 #pragma once
 
+#include "NNEHlslShadersBase.h"
 #include "NNEOperator.h"
-#include "NNEHlslShaderBase.h"
 #include "RenderGraphUtils.h"
 #include "ShaderParameterUtils.h"
 
diff --git a/Engine/Plugins/Experimental/NNERuntimeRDG/Source/NNEHlslShaders/Internal/NNEHlslShadersCastCS.h b/Engine/Plugins/Experimental/NNERuntimeRDG/Source/NNEHlslShaders/Internal/NNEHlslShadersCastCS.h
--- a/Engine/Plugins/Experimental/NNERuntimeRDG/Source/NNEHlslShaders/Internal/NNEHlslShadersCastCS.h
+++ b/Engine/Plugins/Experimental/NNERuntimeRDG/Source/NNEHlslShaders/Internal/NNEHlslShadersCastCS.h
@@ -2,7 +2,7 @@
 
 #pragma once
 
-#include "NNEHlslShaderBase.h"
+#include "NNEHlslShadersBase.h"
 #include "NNEHlslShadersTypeHelper.h"
 #include "NNEOperator.h"
 #include "RenderGraphUtils.h"
diff --git a/Engine/Plugins/Experimental/NNERuntimeRDG/Source/NNEHlslShaders/Internal/NNEHlslShadersConvCS.h b/Engine/Plugins/Experimental/NNERuntimeRDG/Source/NNEHlslShaders/Internal/NNEHlslShadersConvCS.h
--- a/Engine/Plugins/Experimental/NNERuntimeRDG/Source/NNEHlslShaders/Internal/NNEHlslShadersConvCS.h
+++ b/Engine/Plugins/Experimental/NNERuntimeRDG/Source/NNEHlslShaders/Internal/NNEHlslShadersConvCS.h
@@ -2,7 +2,7 @@
 
 #pragma once
 
-#include "NNEHlslShaderBase.h"
+#include "NNEHlslShadersBase.h"
 #include "RenderGraphUtils.h"
 #include "ShaderParameterUtils.h"
 
diff --git a/Engine/Plugins/Experimental/NNERuntimeRDG/Source/NNEHlslShaders/Internal/NNEHlslShadersConvMatmulCS.h b/Engine/Plugins/Experimental/NNERuntimeRDG/Source/NNEHlslShaders/Internal/NNEHlslShadersConvMatmulCS.h
--- a/Engine/Plugins/Experimental/NNERuntimeRDG/Source/NNEHlslShaders/Internal/NNEHlslShadersConvMatmulCS.h
+++ b/Engine/Plugins/Experimental/NNERuntimeRDG/Source/NNEHlslShaders/Internal/NNEHlslShadersConvMatmulCS.h
@@ -2,7 +2,7 @@
 
 #pragma once
 
-#include "NNEHlslShaderBase.h"
+#include "NNEHlslShadersBase.h"
 #include "RenderGraphUtils.h"
 #include "ShaderParameterUtils.h"
 
diff --git a/Engine/Plugins/Experimental/NNERuntimeRDG/Source/NNEHlslShaders/Internal/NNEHlslShadersConvTransposeCS.h b/Engine/Plugins/Experimental/NNERuntimeRDG/Source/NNEHlslShaders/Internal/NNEHlslShadersConvTransposeCS.h
--- a/Engine/Plugins/Experimental/NNERuntimeRDG/Source/NNEHlslShaders/Internal/NNEHlslShadersConvTransposeCS.h
+++ b/Engine/Plugins/Experimental/NNERuntimeRDG/Source/NNEHlslShaders/Internal/NNEHlslShadersConvTransposeCS.h
@@ -2,7 +2,7 @@
 
 #pragma once
 
-#include "NNEHlslShaderBase.h"
+#include "NNEHlslShadersBase.h"
 #include "RenderGraphUtils.h"
 #include "ShaderParameterUtils.h"
 
diff --git a/Engine/Plugins/Experimental/NNERuntimeRDG/Source/NNEHlslShaders/Internal/NNEHlslShadersConvWinogradInputCS.h b/Engine/Plugins/Experimental/NNERuntimeRDG/Source/NNEHlslShaders/Internal/NNEHlslShadersConvWinogradInputCS.h
--- a/Engine/Plugins/Experimental/NNERuntimeRDG/Source/NNEHlslShaders/Internal/NNEHlslShadersConvWinogradInputCS.h
+++ b/Engine/Plugins/Experimental/NNERuntimeRDG/Source/NNEHlslShaders/Internal/NNEHlslShadersConvWinogradInputCS.h
@@ -2,7 +2,7 @@
 
 #pragma once
 
-#include "NNEHlslShaderBase.h"
+#include "NNEHlslShadersBase.h"
 #include "RenderGraphUtils.h"
 
 namespace UE::NNEHlslShaders::Internal
diff --git a/Engine/Plugins/Experimental/NNERuntimeRDG/Source/NNEHlslShaders/Internal/NNEHlslShadersConvWinogradMMMCS.h b/Engine/Plugins/Experimental/NNERuntimeRDG/Source/NNEHlslShaders/Internal/NNEHlslShadersConvWinogradMMMCS.h
--- a/Engine/Plugins/Experimental/NNERuntimeRDG/Source/NNEHlslShaders/Internal/NNEHlslShadersConvWinogradMMMCS.h
+++ b/Engine/Plugins/Experimental/NNERuntimeRDG/Source/NNEHlslShaders/Internal/NNEHlslShadersConvWinogradMMMCS.h
@@ -2,7 +2,7 @@
 
 #pragma once
 
-#include "NNEHlslShaderBase.h"
+#include "NNEHlslShadersBase.h"
 #include "NNEHlslShadersTypeHelper.h"
 #include "RenderGraphUtils.h"
 #include "ShaderParameterUtils.h"
diff --git a/Engine/Plugins/Experimental/NNERuntimeRDG/Source/NNEHlslShaders/Internal/NNEHlslShadersConvWinogradOutputCS.h b/Engine/Plugins/Experimental/NNERuntimeRDG/Source/NNEHlslShaders/Internal/NNEHlslShadersConvWinogradOutputCS.h
--- a/Engine/Plugins/Experimental/NNERuntimeRDG/Source/NNEHlslShaders/Internal/NNEHlslShadersConvWinogradOutputCS.h
+++ b/Engine/Plugins/Experimental/NNERuntimeRDG/Source/NNEHlslShaders/Internal/NNEHlslShadersConvWinogradOutputCS.h
@@ -2,7 +2,7 @@
 
 #pragma once
 
-#include "NNEHlslShaderBase.h"
+#include "NNEHlslShadersBase.h"
 #include "NNEHlslShadersTypeHelper.h"
 #include "RenderGraphUtils.h"
 #include "ShaderParameterUtils.h"
diff --git a/Engine/Plugins/Experimental/NNERuntimeRDG/Source/NNEHlslShaders/Internal/NNEHlslShadersConvWinogradWeightsCS.h b/Engine/Plugins/Experimental/NNERuntimeRDG/Source/NNEHlslShaders/Internal/NNEHlslShadersConvWinogradWeightsCS.h
--- a/Engine/Plugins/Experimental/NNERuntimeRDG/Source/NNEHlslShaders/Internal/NNEHlslShadersConvWinogradWeightsCS.h
+++ b/Engine/Plugins/Experimental/NNERuntimeRDG/Source/NNEHlslShaders/Internal/NNEHlslShadersConvWinogradWeightsCS.h
@@ -2,7 +2,7 @@
 
 #pragma once
 
-#include "NNEHlslShaderBase.h"
+#include "NNEHlslShadersBase.h"
 #include "RenderGraphUtils.h"
 
 namespace UE::NNEHlslShaders::Internal
diff --git a/Engine/Plugins/Experimental/NNERuntimeRDG/Source/NNEHlslShaders/Internal/NNEHlslShadersCumSumCS.h b/Engine/Plugins/Experimental/NNERuntimeRDG/Source/NNEHlslShaders/Internal/NNEHlslShadersCumSumCS.h
--- a/Engine/Plugins/Experimental/NNERuntimeRDG/Source/NNEHlslShaders/Internal/NNEHlslShadersCumSumCS.h
+++ b/Engine/Plugins/Experimental/NNERuntimeRDG/Source/NNEHlslShaders/Internal/NNEHlslShadersCumSumCS.h
@@ -3,7 +3,7 @@
 #pragma once
 
 #include "CoreMinimal.h"
-#include "NNEHlslShaderBase.h"
+#include "NNEHlslShadersBase.h"
 #include "RenderGraphFwd.h"
 #include "RenderGraphUtils.h"
 #include "RHIGlobals.h"
diff --git a/Engine/Plugins/Experimental/NNERuntimeRDG/Source/NNEHlslShaders/Internal/NNEHlslShadersElementWiseBinaryCS.h b/Engine/Plugins/Experimental/NNERuntimeRDG/Source/NNEHlslShaders/Internal/NNEHlslShadersElementWiseBinaryCS.h
--- a/Engine/Plugins/Experimental/NNERuntimeRDG/Source/NNEHlslShaders/Internal/NNEHlslShadersElementWiseBinaryCS.h
+++ b/Engine/Plugins/Experimental/NNERuntimeRDG/Source/NNEHlslShaders/Internal/NNEHlslShadersElementWiseBinaryCS.h
@@ -3,7 +3,7 @@
 #pragma once
 
 #include "NNEOperator.h"
-#include "NNEHlslShaderBase.h"
+#include "NNEHlslShadersBase.h"
 #include "RenderGraphUtils.h"
 #include "ShaderParameterUtils.h"
 
diff --git a/Engine/Plugins/Experimental/NNERuntimeRDG/Source/NNEHlslShaders/Internal/NNEHlslShadersElementWiseUnaryCS.h b/Engine/Plugins/Experimental/NNERuntimeRDG/Source/NNEHlslShaders/Internal/NNEHlslShadersElementWiseUnaryCS.h
--- a/Engine/Plugins/Experimental/NNERuntimeRDG/Source/NNEHlslShaders/Internal/NNEHlslShadersElementWiseUnaryCS.h
+++ b/Engine/Plugins/Experimental/NNERuntimeRDG/Source/NNEHlslShaders/Internal/NNEHlslShadersElementWiseUnaryCS.h
@@ -3,7 +3,7 @@
 #pragma once
 
 #include "NNEOperator.h"
-#include "NNEHlslShaderBase.h"
+#include "NNEHlslShadersBase.h"
 #include "RenderGraphUtils.h"
 #include "ShaderParameterUtils.h"
 
diff --git a/Engine/Plugins/Experimental/NNERuntimeRDG/Source/NNEHlslShaders/Internal/NNEHlslShadersElementWiseVariadicCS.h b/Engine/Plugins/Experimental/NNERuntimeRDG/Source/NNEHlslShaders/Internal/NNEHlslShadersElementWiseVariadicCS.h
--- a/Engine/Plugins/Experimental/NNERuntimeRDG/Source/NNEHlslShaders/Internal/NNEHlslShadersElementWiseVariadicCS.h
+++ b/Engine/Plugins/Experimental/NNERuntimeRDG/Source/NNEHlslShaders/Internal/NNEHlslShadersElementWiseVariadicCS.h
@@ -3,7 +3,7 @@
 #pragma once
 
 #include "NNEOperator.h"
-#include "NNEHlslShaderBase.h"
+#include "NNEHlslShadersBase.h"
 #include "RenderGraphUtils.h"
 #include "ShaderParameterUtils.h"
 
diff --git a/Engine/Plugins/Experimental/NNERuntimeRDG/Source/NNEHlslShaders/Internal/NNEHlslShadersGatherCS.h b/Engine/Plugins/Experimental/NNERuntimeRDG/Source/NNEHlslShaders/Internal/NNEHlslShadersGatherCS.h
--- a/Engine/Plugins/Experimental/NNERuntimeRDG/Source/NNEHlslShaders/Internal/NNEHlslShadersGatherCS.h
+++ b/Engine/Plugins/Experimental/NNERuntimeRDG/Source/NNEHlslShaders/Internal/NNEHlslShadersGatherCS.h
@@ -2,7 +2,7 @@
 
 #pragma once
 
-#include "NNEHlslShaderBase.h"
+#include "NNEHlslShadersBase.h"
 #include "RenderGraphUtils.h"
 #include "ShaderParameterUtils.h"
 
diff --git a/Engine/Plugins/Experimental/NNERuntimeRDG/Source/NNEHlslShaders/Internal/NNEHlslShadersGatherElementsCS.h b/Engine/Plugins/Experimental/NNERuntimeRDG/Source/NNEHlslShaders/Internal/NNEHlslShadersGatherElementsCS.h
--- a/Engine/Plugins/Experimental/NNERuntimeRDG/Source/NNEHlslShaders/Internal/NNEHlslShadersGatherElementsCS.h
+++ b/Engine/Plugins/Experimental/NNERuntimeRDG/Source/NNEHlslShaders/Internal/NNEHlslShadersGatherElementsCS.h
@@ -2,7 +2,7 @@
 
 #pragma once
 
-#include "NNEHlslShaderBase.h"
+#include "NNEHlslShadersBase.h"
 #include "RenderGraphUtils.h"
 #include "ShaderParameterUtils.h"
 
diff --git a/Engine/Plugins/Experimental/NNERuntimeRDG/Source/NNEHlslShaders/Internal/NNEHlslShadersGemmCS.h b/Engine/Plugins/Experimental/NNERuntimeRDG/Source/NNEHlslShaders/Internal/NNEHlslShadersGemmCS.h
--- a/Engine/Plugins/Experimental/NNERuntimeRDG/Source/NNEHlslShaders/Internal/NNEHlslShadersGemmCS.h
+++ b/Engine/Plugins/Experimental/NNERuntimeRDG/Source/NNEHlslShaders/Internal/NNEHlslShadersGemmCS.h
@@ -2,7 +2,7 @@
 
 #pragma once
 
-#include "NNEHlslShaderBase.h"
+#include "NNEHlslShadersBase.h"
 #include "RenderGraphUtils.h"
 #include "ShaderParameterUtils.h"
 
diff --git a/Engine/Plugins/Experimental/NNERuntimeRDG/Source/NNEHlslShaders/Internal/NNEHlslShadersInstanceNormalizationCS.h b/Engine/Plugins/Experimental/NNERuntimeRDG/Source/NNEHlslShaders/Internal/NNEHlslShadersInstanceNormalizationCS.h
--- a/Engine/Plugins/Experimental/NNERuntimeRDG/Source/NNEHlslShaders/Internal/NNEHlslShadersInstanceNormalizationCS.h
+++ b/Engine/Plugins/Experimental/NNERuntimeRDG/Source/NNEHlslShaders/Internal/NNEHlslShadersInstanceNormalizationCS.h
@@ -2,7 +2,7 @@
 
 #pragma once
 
-#include "NNEHlslShaderBase.h"
+#include "NNEHlslShadersBase.h"
 #include "RenderGraphUtils.h"
 #include "ShaderParameterUtils.h"
 
diff --git a/Engine/Plugins/Experimental/NNERuntimeRDG/Source/NNEHlslShaders/Internal/NNEHlslShadersLayerNormalizationCS.h b/Engine/Plugins/Experimental/NNERuntimeRDG/Source/NNEHlslShaders/Internal/NNEHlslShadersLayerNormalizationCS.h
--- a/Engine/Plugins/Experimental/NNERuntimeRDG/Source/NNEHlslShaders/Internal/NNEHlslShadersLayerNormalizationCS.h
+++ b/Engine/Plugins/Experimental/NNERuntimeRDG/Source/NNEHlslShaders/Internal/NNEHlslShadersLayerNormalizationCS.h
@@ -2,7 +2,7 @@
 
 #pragma once
 
-#include "NNEHlslShaderBase.h"
+#include "NNEHlslShadersBase.h"
 #include "RenderGraphUtils.h"
 #include "ShaderParameterUtils.h"
 
diff --git a/Engine/Plugins/Experimental/NNERuntimeRDG/Source/NNEHlslShaders/Internal/NNEHlslShadersLog.h b/Engine/Plugins/Experimental/NNERuntimeRDG/Source/NNEHlslShaders/Internal/NNEHlslShadersLog.h
new file mode 100644
--- /dev/null
+++ b/Engine/Plugins/Experimental/NNERuntimeRDG/Source/NNEHlslShaders/Internal/NNEHlslShadersLog.h
@@ -0,0 +1,7 @@
+// Copyright Epic Games, Inc. All Rights Reserved.
+
+#pragma once
+
+#include "Logging/LogMacros.h"
+
+NNEHLSLSHADERS_API DECLARE_LOG_CATEGORY_EXTERN(LogNNERuntimeRDGHlsl, Log, All);
\ No newline at end of file
diff --git a/Engine/Plugins/Experimental/NNERuntimeRDG/Source/NNEHlslShaders/Internal/NNEHlslShadersNeuralPostProcessingCS.h b/Engine/Plugins/Experimental/NNERuntimeRDG/Source/NNEHlslShaders/Internal/NNEHlslShadersNeuralPostProcessingCS.h
--- a/Engine/Plugins/Experimental/NNERuntimeRDG/Source/NNEHlslShaders/Internal/NNEHlslShadersNeuralPostProcessingCS.h
+++ b/Engine/Plugins/Experimental/NNERuntimeRDG/Source/NNEHlslShaders/Internal/NNEHlslShadersNeuralPostProcessingCS.h
@@ -2,7 +2,7 @@
 
 #pragma once
 
-#include "NNEHlslShaderBase.h"
+#include "NNEHlslShadersBase.h"
 #include "RenderGraphUtils.h"
 #include "ShaderParameterUtils.h"
 
diff --git a/Engine/Plugins/Experimental/NNERuntimeRDG/Source/NNEHlslShaders/Internal/NNEHlslShadersPadCS.h b/Engine/Plugins/Experimental/NNERuntimeRDG/Source/NNEHlslShaders/Internal/NNEHlslShadersPadCS.h
--- a/Engine/Plugins/Experimental/NNERuntimeRDG/Source/NNEHlslShaders/Internal/NNEHlslShadersPadCS.h
+++ b/Engine/Plugins/Experimental/NNERuntimeRDG/Source/NNEHlslShaders/Internal/NNEHlslShadersPadCS.h
@@ -2,7 +2,7 @@
 
 #pragma once
 
-#include "NNEHlslShaderBase.h"
+#include "NNEHlslShadersBase.h"
 #include "RenderGraphUtils.h"
 #include "ShaderParameterUtils.h"
 
diff --git a/Engine/Plugins/Experimental/NNERuntimeRDG/Source/NNEHlslShaders/Internal/NNEHlslShadersPoolCS.h b/Engine/Plugins/Experimental/NNERuntimeRDG/Source/NNEHlslShaders/Internal/NNEHlslShadersPoolCS.h
--- a/Engine/Plugins/Experimental/NNERuntimeRDG/Source/NNEHlslShaders/Internal/NNEHlslShadersPoolCS.h
+++ b/Engine/Plugins/Experimental/NNERuntimeRDG/Source/NNEHlslShaders/Internal/NNEHlslShadersPoolCS.h
@@ -2,7 +2,7 @@
 
 #pragma once
 
-#include "NNEHlslShaderBase.h"
+#include "NNEHlslShadersBase.h"
 #include "RenderGraphUtils.h"
 #include "ShaderParameterUtils.h"
 
diff --git a/Engine/Plugins/Experimental/NNERuntimeRDG/Source/NNEHlslShaders/Internal/NNEHlslShadersReduceCS.h b/Engine/Plugins/Experimental/NNERuntimeRDG/Source/NNEHlslShaders/Internal/NNEHlslShadersReduceCS.h
--- a/Engine/Plugins/Experimental/NNERuntimeRDG/Source/NNEHlslShaders/Internal/NNEHlslShadersReduceCS.h
+++ b/Engine/Plugins/Experimental/NNERuntimeRDG/Source/NNEHlslShaders/Internal/NNEHlslShadersReduceCS.h
@@ -3,7 +3,7 @@
 #pragma once
 
 #include "CoreMinimal.h"
-#include "NNEHlslShaderBase.h"
+#include "NNEHlslShadersBase.h"
 #include "RenderGraphFwd.h"
 #include "RenderGraphUtils.h"
 #include "ShaderParameterUtils.h"
diff --git a/Engine/Plugins/Experimental/NNERuntimeRDG/Source/NNEHlslShaders/Internal/NNEHlslShadersResizeCS.h b/Engine/Plugins/Experimental/NNERuntimeRDG/Source/NNEHlslShaders/Internal/NNEHlslShadersResizeCS.h
--- a/Engine/Plugins/Experimental/NNERuntimeRDG/Source/NNEHlslShaders/Internal/NNEHlslShadersResizeCS.h
+++ b/Engine/Plugins/Experimental/NNERuntimeRDG/Source/NNEHlslShaders/Internal/NNEHlslShadersResizeCS.h
@@ -2,7 +2,7 @@
 
 #pragma once
 
-#include "NNEHlslShaderBase.h"
+#include "NNEHlslShadersBase.h"
 #include "RenderGraphUtils.h"
 #include "ShaderParameterUtils.h"
 
diff --git a/Engine/Plugins/Experimental/NNERuntimeRDG/Source/NNEHlslShaders/Internal/NNEHlslShadersScatterNDCS.h b/Engine/Plugins/Experimental/NNERuntimeRDG/Source/NNEHlslShaders/Internal/NNEHlslShadersScatterNDCS.h
--- a/Engine/Plugins/Experimental/NNERuntimeRDG/Source/NNEHlslShaders/Internal/NNEHlslShadersScatterNDCS.h
+++ b/Engine/Plugins/Experimental/NNERuntimeRDG/Source/NNEHlslShaders/Internal/NNEHlslShadersScatterNDCS.h
@@ -2,7 +2,7 @@
 
 #pragma once
 
-#include "NNEHlslShaderBase.h"
+#include "NNEHlslShadersBase.h"
 #include "RenderGraphUtils.h"
 #include "ShaderParameterUtils.h"
 
diff --git a/Engine/Plugins/Experimental/NNERuntimeRDG/Source/NNEHlslShaders/Internal/NNEHlslShadersSliceCS.h b/Engine/Plugins/Experimental/NNERuntimeRDG/Source/NNEHlslShaders/Internal/NNEHlslShadersSliceCS.h
--- a/Engine/Plugins/Experimental/NNERuntimeRDG/Source/NNEHlslShaders/Internal/NNEHlslShadersSliceCS.h
+++ b/Engine/Plugins/Experimental/NNERuntimeRDG/Source/NNEHlslShaders/Internal/NNEHlslShadersSliceCS.h
@@ -2,7 +2,7 @@
 
 #pragma once
 
-#include "NNEHlslShaderBase.h"
+#include "NNEHlslShadersBase.h"
 #include "RenderGraphUtils.h"
 #include "ShaderParameterUtils.h"
 
diff --git a/Engine/Plugins/Experimental/NNERuntimeRDG/Source/NNEHlslShaders/Internal/NNEHlslShadersSoftmaxCS.h b/Engine/Plugins/Experimental/NNERuntimeRDG/Source/NNEHlslShaders/Internal/NNEHlslShadersSoftmaxCS.h
--- a/Engine/Plugins/Experimental/NNERuntimeRDG/Source/NNEHlslShaders/Internal/NNEHlslShadersSoftmaxCS.h
+++ b/Engine/Plugins/Experimental/NNERuntimeRDG/Source/NNEHlslShaders/Internal/NNEHlslShadersSoftmaxCS.h
@@ -2,7 +2,7 @@
 
 #pragma once
 
-#include "NNEHlslShaderBase.h"
+#include "NNEHlslShadersBase.h"
 #include "RenderGraphUtils.h"
 #include "ShaderParameterUtils.h"
 
diff --git a/Engine/Plugins/Experimental/NNERuntimeRDG/Source/NNEHlslShaders/Internal/NNEHlslShadersSplitCS.h b/Engine/Plugins/Experimental/NNERuntimeRDG/Source/NNEHlslShaders/Internal/NNEHlslShadersSplitCS.h
--- a/Engine/Plugins/Experimental/NNERuntimeRDG/Source/NNEHlslShaders/Internal/NNEHlslShadersSplitCS.h
+++ b/Engine/Plugins/Experimental/NNERuntimeRDG/Source/NNEHlslShaders/Internal/NNEHlslShadersSplitCS.h
@@ -2,7 +2,7 @@
 
 #pragma once
 
-#include "NNEHlslShaderBase.h"
+#include "NNEHlslShadersBase.h"
 #include "RenderGraphUtils.h"
 #include "ShaderParameterUtils.h"
 
diff --git a/Engine/Plugins/Experimental/NNERuntimeRDG/Source/NNEHlslShaders/Internal/NNEHlslShadersTransposeCS.h b/Engine/Plugins/Experimental/NNERuntimeRDG/Source/NNEHlslShaders/Internal/NNEHlslShadersTransposeCS.h
--- a/Engine/Plugins/Experimental/NNERuntimeRDG/Source/NNEHlslShaders/Internal/NNEHlslShadersTransposeCS.h
+++ b/Engine/Plugins/Experimental/NNERuntimeRDG/Source/NNEHlslShaders/Internal/NNEHlslShadersTransposeCS.h
@@ -2,7 +2,7 @@
 
 #pragma once
 
-#include "NNEHlslShaderBase.h"
+#include "NNEHlslShadersBase.h"
 #include "RenderGraphUtils.h"
 #include "ShaderParameterUtils.h"
 
diff --git a/Engine/Plugins/Experimental/NNERuntimeRDG/Source/NNEHlslShaders/Internal/NNEHlslShadersUpsampleCS.h b/Engine/Plugins/Experimental/NNERuntimeRDG/Source/NNEHlslShaders/Internal/NNEHlslShadersUpsampleCS.h
--- a/Engine/Plugins/Experimental/NNERuntimeRDG/Source/NNEHlslShaders/Internal/NNEHlslShadersUpsampleCS.h
+++ b/Engine/Plugins/Experimental/NNERuntimeRDG/Source/NNEHlslShaders/Internal/NNEHlslShadersUpsampleCS.h
@@ -2,7 +2,7 @@
 
 #pragma once
 
-#include "NNEHlslShaderBase.h"
+#include "NNEHlslShadersBase.h"
 #include "RenderGraphUtils.h"
 #include "ShaderParameterUtils.h"
 
diff --git a/Engine/Plugins/Experimental/NNERuntimeRDG/Source/NNEHlslShaders/Private/NNEHlslShaderBase.cpp b/Engine/Plugins/Experimental/NNERuntimeRDG/Source/NNEHlslShaders/Private/NNEHlslShaderBase.cpp
--- a/Engine/Plugins/Experimental/NNERuntimeRDG/Source/NNEHlslShaders/Private/NNEHlslShaderBase.cpp
+++ b/Engine/Plugins/Experimental/NNERuntimeRDG/Source/NNEHlslShaders/Private/NNEHlslShaderBase.cpp
@@ -1,6 +1,6 @@
 // Copyright Epic Games, Inc. All Rights Reserved.
 
-#include "NNEHlslShaderBase.h"
+#include "NNEHlslShadersBase.h"
 #include "DataDrivenShaderPlatformInfo.h"
 
 namespace UE::NNEHlslShaders::Internal
diff --git a/Engine/Plugins/Experimental/NNERuntimeRDG/Source/NNEHlslShaders/Private/NNEHlslShadersElementWiseBinaryCS.cpp b/Engine/Plugins/Experimental/NNERuntimeRDG/Source/NNEHlslShaders/Private/NNEHlslShadersElementWiseBinaryCS.cpp
--- a/Engine/Plugins/Experimental/NNERuntimeRDG/Source/NNEHlslShaders/Private/NNEHlslShadersElementWiseBinaryCS.cpp
+++ b/Engine/Plugins/Experimental/NNERuntimeRDG/Source/NNEHlslShaders/Private/NNEHlslShadersElementWiseBinaryCS.cpp
@@ -1,7 +1,7 @@
 // Copyright Epic Games, Inc. All Rights Reserved.
 
 #include "NNEHlslShadersElementWiseBinaryCS.h"
-#include "NNE.h"
+#include "NNEHlslShadersLog.h"
 
 namespace UE::NNEHlslShaders::Internal
 {
@@ -49,7 +49,7 @@ namespace UE::NNEHlslShaders::Internal
 
 		if (OpFunc == "")
 		{
-			UE_LOG(LogNNE, Warning, TEXT("Undefined ElementWise Binary operator name for operator:%d"), int(OpType));
+			UE_LOG(LogNNERuntimeRDGHlsl, Warning, TEXT("Undefined ElementWise Binary operator name for operator:%d"), int(OpType));
 		}
 
 		return OpFunc;
diff --git a/Engine/Plugins/Experimental/NNERuntimeRDG/Source/NNEHlslShaders/Private/NNEHlslShadersElementWiseUnaryCS.cpp b/Engine/Plugins/Experimental/NNERuntimeRDG/Source/NNEHlslShaders/Private/NNEHlslShadersElementWiseUnaryCS.cpp
--- a/Engine/Plugins/Experimental/NNERuntimeRDG/Source/NNEHlslShaders/Private/NNEHlslShadersElementWiseUnaryCS.cpp
+++ b/Engine/Plugins/Experimental/NNERuntimeRDG/Source/NNEHlslShaders/Private/NNEHlslShadersElementWiseUnaryCS.cpp
@@ -1,8 +1,7 @@
 // Copyright Epic Games, Inc. All Rights Reserved.
 
 #include "NNEHlslShadersElementWiseUnaryCS.h"
-
-#include "NNE.h"
+#include "NNEHlslShadersLog.h"
 
 namespace UE::NNEHlslShaders::Internal
 {
@@ -89,7 +88,7 @@ namespace UE::NNEHlslShaders::Internal
 
 		if (OpFunc == FString(""))
 		{
-			UE_LOG(LogNNE, Warning, TEXT("Undefined ElementWise Unary operator name for operator:%d"), int(OpType));
+			UE_LOG(LogNNERuntimeRDGHlsl, Warning, TEXT("Undefined ElementWise Unary operator name for operator:%d"), int(OpType));
 		}
 
 		return OpFunc;
diff --git a/Engine/Plugins/Experimental/NNERuntimeRDG/Source/NNEHlslShaders/Private/NNEHlslShadersElementWiseVariadicCS.cpp b/Engine/Plugins/Experimental/NNERuntimeRDG/Source/NNEHlslShaders/Private/NNEHlslShadersElementWiseVariadicCS.cpp
--- a/Engine/Plugins/Experimental/NNERuntimeRDG/Source/NNEHlslShaders/Private/NNEHlslShadersElementWiseVariadicCS.cpp
+++ b/Engine/Plugins/Experimental/NNERuntimeRDG/Source/NNEHlslShaders/Private/NNEHlslShadersElementWiseVariadicCS.cpp
@@ -1,7 +1,7 @@
 // Copyright Epic Games, Inc. All Rights Reserved.
 
 #include "NNEHlslShadersElementWiseVariadicCS.h"
-#include "NNE.h"
+#include "NNEHlslShadersLog.h"
 
 namespace UE::NNEHlslShaders::Internal
 {
@@ -38,7 +38,7 @@ namespace UE::NNEHlslShaders::Internal
 
 		if (OpFunc == "")
 		{
-			UE_LOG(LogNNE, Warning, TEXT("Undefined ElementWise Variadic operator name for operator:%d"), int(OpType));
+			UE_LOG(LogNNERuntimeRDGHlsl, Warning, TEXT("Undefined ElementWise Variadic operator name for operator:%d"), int(OpType));
 		}
 
 		return OpFunc;
diff --git a/Engine/Plugins/Experimental/NNERuntimeRDG/Source/NNEHlslShaders/Private/NNEHlslShadersLog.cpp b/Engine/Plugins/Experimental/NNERuntimeRDG/Source/NNEHlslShaders/Private/NNEHlslShadersLog.cpp
new file mode 100644
--- /dev/null
+++ b/Engine/Plugins/Experimental/NNERuntimeRDG/Source/NNEHlslShaders/Private/NNEHlslShadersLog.cpp
@@ -0,0 +1,5 @@
+// Copyright Epic Games, Inc. All Rights Reserved.
+
+#include "NNEHlslShadersLog.h"
+
+DEFINE_LOG_CATEGORY(LogNNERuntimeRDGHlsl);
\ No newline at end of file
diff --git a/Engine/Plugins/Experimental/NNERuntimeRDG/Source/NNEHlslShaders/Private/NNEHlslShadersModule.cpp b/Engine/Plugins/Experimental/NNERuntimeRDG/Source/NNEHlslShaders/Private/NNEHlslShadersModule.cpp
--- a/Engine/Plugins/Experimental/NNERuntimeRDG/Source/NNEHlslShaders/Private/NNEHlslShadersModule.cpp
+++ b/Engine/Plugins/Experimental/NNERuntimeRDG/Source/NNEHlslShaders/Private/NNEHlslShadersModule.cpp
@@ -4,7 +4,7 @@
 #include "Interfaces/IPluginManager.h"
 #include "Misc/Paths.h"
 #include "Modules/ModuleManager.h"
-#include "NNE.h"
+#include "NNEHlslShadersLog.h"
 #include "ShaderCore.h"
 
 class FNNEHlslShadersModule : public IModuleInterface
@@ -21,7 +21,7 @@ public:
 		}
 		else
 		{
-			UE_LOG(LogNNE, Warning, TEXT("Shaders directory not added. Failed to find NNERuntimeRDG plugin"));
+			UE_LOG(LogNNERuntimeRDGHlsl, Warning, TEXT("Shaders directory not added. Failed to find NNERuntimeRDG plugin"));
 		}
 
 	}
diff --git a/Engine/Plugins/Experimental/NNERuntimeRDG/Source/NNEHlslShaders/Private/NNEHlslShadersScatterND.cpp b/Engine/Plugins/Experimental/NNERuntimeRDG/Source/NNEHlslShaders/Private/NNEHlslShadersScatterND.cpp
deleted file mode 100644
diff --git a/Engine/Plugins/Experimental/NNERuntimeRDG/Source/NNEHlslShaders/Private/NNEHlslShadersScatterNDCS.cpp b/Engine/Plugins/Experimental/NNERuntimeRDG/Source/NNEHlslShaders/Private/NNEHlslShadersScatterNDCS.cpp
new file mode 100644
--- /dev/null
+++ b/Engine/Plugins/Experimental/NNERuntimeRDG/Source/NNEHlslShaders/Private/NNEHlslShadersScatterNDCS.cpp
@@ -0,0 +1,29 @@
+// Copyright Epic Games, Inc. All Rights Reserved.
+
+#include "NNEHlslShadersScatterNDCS.h"
+#include "NNE.h"
+
+namespace UE::NNEHlslShaders::Internal
+{
+	void FScatterNDCS::ModifyCompilationEnvironment(const FGlobalShaderPermutationParameters& InParameters, FShaderCompilerEnvironment& OutEnvironment)
+	{
+		FGlobalShader::ModifyCompilationEnvironment(InParameters, OutEnvironment);
+
+		OutEnvironment.SetDefine(TEXT("THREADGROUP_SIZE"), FScatterNDConstants::NUM_GROUP_THREADS);
+
+	}
+
+    EScatterNDReductionType FScatterNDCS::ReductionFromString(const TCHAR* StringVal)
+	{
+		EScatterNDReductionType OutValue = EScatterNDReductionType::None;
+		if (FCString::Stricmp(StringVal, TEXT("none")) == 0) OutValue = EScatterNDReductionType::None;
+		else if (FCString::Stricmp(StringVal, TEXT("add")) == 0) OutValue = EScatterNDReductionType::Add;
+		else if (FCString::Stricmp(StringVal, TEXT("mul")) == 0) OutValue = EScatterNDReductionType::Mul;
+		else if (FCString::Stricmp(StringVal, TEXT("max")) == 0) OutValue = EScatterNDReductionType::Max;
+		else if (FCString::Stricmp(StringVal, TEXT("min")) == 0) OutValue = EScatterNDReductionType::Min;
+
+        return OutValue;
+	}
+
+	IMPLEMENT_GLOBAL_SHADER(FScatterNDCS, "/NNEHlslShaders/NNEHlslShadersScatterND.usf", "ScatterND", SF_Compute);
+} // UE::NNEHlslShaders::Internal
\ No newline at end of file
diff --git a/Engine/Plugins/Experimental/NNERuntimeRDG/Source/NNERuntimeRDG/Internal/NNERenderGraphUtils.h b/Engine/Plugins/Experimental/NNERuntimeRDG/Source/NNERuntimeRDG/Internal/NNERenderGraphUtils.h
deleted file mode 100644
diff --git a/Engine/Plugins/Experimental/NNERuntimeRDG/Source/NNERuntimeRDG/Private/Helper/NNERuntimeRDGHelperReshape.h b/Engine/Plugins/Experimental/NNERuntimeRDG/Source/NNERuntimeRDG/Private/Helper/NNERuntimeRDGHelperReshape.h
--- a/Engine/Plugins/Experimental/NNERuntimeRDG/Source/NNERuntimeRDG/Private/Helper/NNERuntimeRDGHelperReshape.h
+++ b/Engine/Plugins/Experimental/NNERuntimeRDG/Source/NNERuntimeRDG/Private/Helper/NNERuntimeRDGHelperReshape.h
@@ -6,7 +6,7 @@
 #include "Algo/Count.h"
 #include "CoreMinimal.h"
 #include "Misc/EnumerateRange.h"
-#include "NNE.h"
+#include "NNEHlslShadersLog.h"
 #include "NNETensor.h"
 #include "NNETypes.h"
 
@@ -27,14 +27,14 @@ namespace UE::NNERuntimeRDG::Private::ShapeHelper::Reshape
 			// at most 1 dimension can be -1
 			if (Algo::Count(ReshapedShape, -1) > 1)
 			{
-				UE_LOG(LogNNE, Error, TEXT("Reshape: Shape tensor can't contain more than one '-1'."));
+				UE_LOG(LogNNERuntimeRDGHlsl, Error, TEXT("Reshape: Shape tensor can't contain more than one '-1'."));
 				return false;
 			}
 			for (TEnumerateRef<DataType> Elem : EnumerateRange(ReshapedShape))
 			{
 				if (!(*Elem != 0 || InputTensorShape.Rank() > Elem.GetIndex()))
 				{
-					UE_LOG(LogNNE, Error, TEXT("Reshape: Shape tensor contains '0' in an invalid place."));
+					UE_LOG(LogNNERuntimeRDGHlsl, Error, TEXT("Reshape: Shape tensor contains '0' in an invalid place."));
 					return false;
 				}
 				*Elem =
@@ -49,7 +49,7 @@ namespace UE::NNERuntimeRDG::Private::ShapeHelper::Reshape
 			// no -1 is allowed if there is a 0
 			if (Algo::Count(ReshapedShape, 0) != 0 && Algo::Count(ReshapedShape, -1) != 0)
 			{
-				UE_LOG(LogNNE, Error, TEXT("Reshape: Shape tensor contains both '0' and '-1'. This is not allowed."));
+				UE_LOG(LogNNERuntimeRDGHlsl, Error, TEXT("Reshape: Shape tensor contains both '0' and '-1'. This is not allowed."));
 				return false;
 			}
 		}
diff --git a/Engine/Plugins/Experimental/NNERuntimeRDG/Source/NNERuntimeRDG/Private/Hlsl/NNERuntimeRDGBatchNormalization.cpp b/Engine/Plugins/Experimental/NNERuntimeRDG/Source/NNERuntimeRDG/Private/Hlsl/NNERuntimeRDGBatchNormalization.cpp
--- a/Engine/Plugins/Experimental/NNERuntimeRDG/Source/NNERuntimeRDG/Private/Hlsl/NNERuntimeRDGBatchNormalization.cpp
+++ b/Engine/Plugins/Experimental/NNERuntimeRDG/Source/NNERuntimeRDG/Private/Hlsl/NNERuntimeRDGBatchNormalization.cpp
@@ -1,7 +1,9 @@
 // Copyright Epic Games, Inc. All Rights Reserved.
 
 #include "NNERuntimeRDGBatchNormalization.h"
+
 #include "NNEHlslShadersBatchNormalizationCS.h"
+#include "NNEHlslShadersLog.h"
 #include "NNERuntimeRDGHlslHelper.h"
 #include "NNEAttributeMap.h"
 #include "NNETypes.h"
@@ -46,7 +48,7 @@ namespace UE::NNERuntimeRDG::Private::Hlsl
 
 			if (OutputTensorDescs.Num() > 1)
 			{
-				UE_LOG(LogNNE, Warning, TEXT("BatchNormalization is only supported in inference mode at the moment, supporting no more than one output."));
+				UE_LOG(LogNNERuntimeRDGHlsl, Warning, TEXT("BatchNormalization: Only supported in inference mode at the moment, supporting no more than one output."));
 				return false;
 			}
 			
diff --git a/Engine/Plugins/Experimental/NNERuntimeRDG/Source/NNERuntimeRDG/Private/Hlsl/NNERuntimeRDGCast.cpp b/Engine/Plugins/Experimental/NNERuntimeRDG/Source/NNERuntimeRDG/Private/Hlsl/NNERuntimeRDGCast.cpp
--- a/Engine/Plugins/Experimental/NNERuntimeRDG/Source/NNERuntimeRDG/Private/Hlsl/NNERuntimeRDGCast.cpp
+++ b/Engine/Plugins/Experimental/NNERuntimeRDG/Source/NNERuntimeRDG/Private/Hlsl/NNERuntimeRDGCast.cpp
@@ -5,6 +5,7 @@
 #include "Helper/NNERuntimeRDGLogHelper.h"
 #include "Helper/NNERuntimeRDGOperatorHelper.h"
 #include "NNEHlslShadersCastCS.h"
+#include "NNEHlslShadersLog.h"
 #include "NNEHlslShadersTypeHelper.h"
 #include "NNERuntimeRDGHelperCast.h"
 #include "NNERuntimeRDGHlslHelper.h"
@@ -39,7 +40,7 @@ namespace UE::NNERuntimeRDG::Private::Hlsl
 
 			if (!Output.HasPreparedData() && Has64BitDataType)
 			{
-				UE_LOG(LogNNE, Warning, TEXT("Cast: Output could not be computed as a constant tensor, however Cast doesn't support dynamic 64 bit tensor types."));
+				UE_LOG(LogNNERuntimeRDGHlsl, Warning, TEXT("Cast: Output could not be computed as a constant tensor, however Cast doesn't support dynamic 64 bit tensor types."));
 				return -1;
 			}
 
@@ -56,7 +57,7 @@ namespace UE::NNERuntimeRDG::Private::Hlsl
 
 			if (ToFromAttribute != ToFromTensor)
 			{
-				UE_LOG(LogNNE, Warning, TEXT("Cast should output a tensor of type %d but was of type %d."), int(ToFromAttribute), int(ToFromTensor));
+				UE_LOG(LogNNERuntimeRDGHlsl, Warning, TEXT("Cast: Should output a tensor of type %d but was of type %d."), int(ToFromAttribute), int(ToFromTensor));
 				return false;
 			}
 			
@@ -126,7 +127,7 @@ namespace UE::NNERuntimeRDG::Private::Hlsl
 					break;
 				default:
 					FString TargetType = LogHelper::GetTensorDataTypeName(To);
-					UE_LOG(LogNNE, Warning, TEXT("Cast: Target tensor data type %s not supported."), *TargetType);
+					UE_LOG(LogNNERuntimeRDGHlsl, Warning, TEXT("Cast: Target tensor data type %s not supported."), *TargetType);
 					bIsValid = false;
 			}
 		}
diff --git a/Engine/Plugins/Experimental/NNERuntimeRDG/Source/NNERuntimeRDG/Private/Hlsl/NNERuntimeRDGConcat.cpp b/Engine/Plugins/Experimental/NNERuntimeRDG/Source/NNERuntimeRDG/Private/Hlsl/NNERuntimeRDGConcat.cpp
--- a/Engine/Plugins/Experimental/NNERuntimeRDG/Source/NNERuntimeRDG/Private/Hlsl/NNERuntimeRDGConcat.cpp
+++ b/Engine/Plugins/Experimental/NNERuntimeRDG/Source/NNERuntimeRDG/Private/Hlsl/NNERuntimeRDGConcat.cpp
@@ -1,6 +1,8 @@
 // Copyright Epic Games, Inc. All Rights Reserved.
 
 #include "NNERuntimeRDGConcat.h"
+
+#include "NNEHlslShadersLog.h"
 #include "NNETensor.h"
 #include "NNETypes.h"
 #include "NNERuntimeRDGHelperConcat.h"
@@ -66,7 +68,7 @@ namespace UE::NNERuntimeRDG::Private::Hlsl
 				{
 					if (r != Axis && (OutputShapeData[r] != InputTensors[i]->GetShape().GetData()[r]))
 					{
-						UE_LOG(LogNNE, Warning, TEXT("Concat: all input tensors should have the same shape except on the concatenation axis"));
+						UE_LOG(LogNNERuntimeRDGHlsl, Warning, TEXT("Concat: All input tensors should have the same shape except on the concatenation axis"));
 						return false;
 					}
 				}
@@ -95,14 +97,14 @@ namespace UE::NNERuntimeRDG::Private::Hlsl
 			{
 				if (InputsRank != InputTensorDescs[i].GetShape().Rank())
 				{
-					UE_LOG(LogNNE, Warning, TEXT("Concat: all input tensors should have the same rank"));
+					UE_LOG(LogNNERuntimeRDGHlsl, Warning, TEXT("Concat: All input tensors should have the same rank"));
 					return false;
 				}
 			}
 			
 			if (Axis < -InputsRank || Axis >(InputsRank - 1))
 			{
-				UE_LOG(LogNNE, Warning, TEXT("Axis should be in range [-r,r-1] however it is %d while inputs have rank %d."), Axis, InputsRank);
+				UE_LOG(LogNNERuntimeRDGHlsl, Warning, TEXT("Concat: Axis should be in range [-r,r-1] however it is %d while inputs have rank %d."), Axis, InputsRank);
 				return false;
 			}
 			
@@ -154,19 +156,19 @@ namespace UE::NNERuntimeRDG::Private::Hlsl
 
 		if (InputTypes.Num() == 0)
 		{
-			UE_LOG(LogNNE, Error, TEXT("Concat operator requires at least 1 input"));
+			UE_LOG(LogNNERuntimeRDGHlsl, Error, TEXT("Concat: Operator requires at least 1 input"));
 			bIsValid = false;
 		}
 		for (int32 i = 0; i < InputTypes.Num(); ++i)
 		{
 			if (InputTypes[i] != ENNETensorDataType::Float && InputTypes[i] != ENNETensorDataType::Half)
 			{
-				UE_LOG(LogNNE, Warning, TEXT("Concat operator input '%d' of type '%d' is not supported, should be float or half at the moment."), i, int(InputTypes[i]));
+				UE_LOG(LogNNERuntimeRDGHlsl, Warning, TEXT("Concat: Operator input '%d' of type '%d' is not supported, should be float or half at the moment."), i, int(InputTypes[i]));
 				bIsValid = false;
 			}
 			if (InputTypes[i] != InputTypes[0])
 			{
-				UE_LOG(LogNNE, Warning, TEXT("Concat operator input '%d' of type '%d' does not match type '%d' of the first input."), i, int(InputTypes[i]), int(InputTypes[0]));
+				UE_LOG(LogNNERuntimeRDGHlsl, Warning, TEXT("Concat: Operator input '%d' of type '%d' does not match type '%d' of the first input."), i, int(InputTypes[i]), int(InputTypes[0]));
 				bIsValid = false;
 			}
 		}
diff --git a/Engine/Plugins/Experimental/NNERuntimeRDG/Source/NNERuntimeRDG/Private/Hlsl/NNERuntimeRDGConstant.cpp b/Engine/Plugins/Experimental/NNERuntimeRDG/Source/NNERuntimeRDG/Private/Hlsl/NNERuntimeRDGConstant.cpp
--- a/Engine/Plugins/Experimental/NNERuntimeRDG/Source/NNERuntimeRDG/Private/Hlsl/NNERuntimeRDGConstant.cpp
+++ b/Engine/Plugins/Experimental/NNERuntimeRDG/Source/NNERuntimeRDG/Private/Hlsl/NNERuntimeRDGConstant.cpp
@@ -5,6 +5,7 @@
 #include "Helper/NNERuntimeRDGLogHelper.h"
 #include "Helper/NNERuntimeRDGOperatorHelper.h"
 #include "NNEAttributeTensor.h"
+#include "NNEHlslShadersLog.h"
 #include "NNERuntimeRDGHlslHelper.h"
 #include "NNETensor.h"
 #include "NNETypes.h"
@@ -45,7 +46,7 @@ namespace UE::NNERuntimeRDG::Private::Hlsl
 				{
 					if (Output.GetDataType() != ENNETensorDataType::Float)
 					{
-						UE_LOG(LogNNE, Warning, TEXT("Output data type %s does not match constant type of float"), 
+						UE_LOG(LogNNERuntimeRDGHlsl, Warning, TEXT("Constant: Output data type %s does not match constant type of float"), 
 							   					*LogHelper::GetTensorDataTypeName(Output.GetDataType()));
 						return -1;
 					}
@@ -59,7 +60,7 @@ namespace UE::NNERuntimeRDG::Private::Hlsl
 				{
 					if (Output.GetDataType() != ENNETensorDataType::Float)
 					{
-						UE_LOG(LogNNE, Warning, TEXT("Output data type %s does not match constant type of float"), 
+						UE_LOG(LogNNERuntimeRDGHlsl, Warning, TEXT("Constant: Output data type %s does not match constant type of float"), 
 							   					*LogHelper::GetTensorDataTypeName(Output.GetDataType()));
 						return -1;
 					}
@@ -74,7 +75,7 @@ namespace UE::NNERuntimeRDG::Private::Hlsl
 					const FAttributeTensor AttributeTensor = Attribute.GetValue<FAttributeTensor>();
 					if (Output.GetDataType() != AttributeTensor.GetDataType())
 					{
-						UE_LOG(LogNNE, Warning, TEXT("Output data type %s does not match constant tensor data type %s"), 
+						UE_LOG(LogNNERuntimeRDGHlsl, Warning, TEXT("Constant: Output data type %s does not match constant tensor data type %s"), 
 							   					*LogHelper::GetTensorDataTypeName(Output.GetDataType()), 
 												*LogHelper::GetTensorDataTypeName(AttributeTensor.GetDataType()));
 						return -1;
@@ -121,7 +122,7 @@ namespace UE::NNERuntimeRDG::Private::Hlsl
 
 		if (AttributeMap.Num() != 1)
 		{
-			UE_LOG(LogNNE, Warning, TEXT("Constant operator requires exacly one attribute, but '%d' attributes found."), AttributeMap.Num());
+			UE_LOG(LogNNERuntimeRDGHlsl, Warning, TEXT("Constant: Operator requires exacly one attribute, but '%d' attributes found."), AttributeMap.Num());
 			return false;
 		}
 
diff --git a/Engine/Plugins/Experimental/NNERuntimeRDG/Source/NNERuntimeRDG/Private/Hlsl/NNERuntimeRDGConv.cpp b/Engine/Plugins/Experimental/NNERuntimeRDG/Source/NNERuntimeRDG/Private/Hlsl/NNERuntimeRDGConv.cpp
--- a/Engine/Plugins/Experimental/NNERuntimeRDG/Source/NNERuntimeRDG/Private/Hlsl/NNERuntimeRDGConv.cpp
+++ b/Engine/Plugins/Experimental/NNERuntimeRDG/Source/NNERuntimeRDG/Private/Hlsl/NNERuntimeRDGConv.cpp
@@ -13,6 +13,7 @@
 #include "NNEHlslShadersConvWinogradMMMCS.h"
 #include "NNEHlslShadersConvWinogradOutputCS.h"
 #include "NNEHlslShadersConvWinogradWeightsCS.h"
+#include "NNEHlslShadersLog.h"
 #include "NNEHlslShadersTypeHelper.h"
 #include "NNERuntimeRDGHelperTranspose.h"
 #include "NNERuntimeRDGHlslHelper.h"
@@ -98,7 +99,7 @@ namespace UE::NNERuntimeRDG::Private::Hlsl
 			GroupSize = FConvCS::GetBiggestCompatibleGroupSize(Weights.GetData(), Dilations, Strides);
 			if(GroupSize == EConvGroupSize::MAX)
 			{
-				UE_LOG(LogNNE, Warning, TEXT("Hlsl Conv: kernel size, strides, dilations combination is not supported. Kernel tensor: %s."), *InputTensors[1]->GetName());
+				UE_LOG(LogNNERuntimeRDGHlsl, Warning, TEXT("Conv: Kernel size, strides, dilations combination is not supported. Kernel tensor: %s."), *InputTensors[1]->GetName());
 				return -1;
 			}
 
@@ -127,17 +128,17 @@ namespace UE::NNERuntimeRDG::Private::Hlsl
 			
 			if (Input.GetShape().Rank() < 2)
 			{
-				UE_LOG(LogNNE, Warning, TEXT("Conv first input should be at least of rank 2"));
+				UE_LOG(LogNNERuntimeRDGHlsl, Warning, TEXT("Conv: First input should be at least of rank 2"));
 				return false;
 			}
 			if (Weights.GetShape().Rank() != Input.GetShape().Rank())
 			{
-				UE_LOG(LogNNE, Warning, TEXT("Conv first and second inputs should be of same ranks"));
+				UE_LOG(LogNNERuntimeRDGHlsl, Warning, TEXT("Conv: First and second inputs should be of same ranks"));
 				return false;
 			}
 			if (Output.GetShape().Rank() != Input.GetShape().Rank())
 			{
-				UE_LOG(LogNNE, Warning, TEXT("Conv first and output should be of same ranks"));
+				UE_LOG(LogNNERuntimeRDGHlsl, Warning, TEXT("Conv: First and output should be of same ranks"));
 				return false;
 			}
 
@@ -160,7 +161,7 @@ namespace UE::NNERuntimeRDG::Private::Hlsl
 			Strides = Attributes.GetValueOrDefault<TArray<int32>>(TEXT("strides"), DilationsOrStridesDefault);
 			if (Strides.Num() != NumDimensions)
 			{
-				UE_LOG(LogNNE, Warning, TEXT("Conv strides count doesn't match the number of spatial dimensions"));
+				UE_LOG(LogNNERuntimeRDGHlsl, Warning, TEXT("Conv: Strides count doesn't match the number of spatial dimensions"));
 				return false;
 			}
 			BufferPixelFormat = TensorDataTypeToPixelFormat(Input.GetDataType());
@@ -611,7 +612,7 @@ namespace UE::NNERuntimeRDG::Private::Hlsl
 		{
 			if (InputType != ENNETensorDataType::None && InputType != InputTypes[0])
 			{
-				UE_LOG(LogNNE, Warning, TEXT("All input tensor data types have to match each other"));
+				UE_LOG(LogNNERuntimeRDGHlsl, Warning, TEXT("Conv: All input tensor data types have to match each other"));
 				return false;
 			}
 		}
diff --git a/Engine/Plugins/Experimental/NNERuntimeRDG/Source/NNERuntimeRDG/Private/Hlsl/NNERuntimeRDGConvTranspose.cpp b/Engine/Plugins/Experimental/NNERuntimeRDG/Source/NNERuntimeRDG/Private/Hlsl/NNERuntimeRDGConvTranspose.cpp
--- a/Engine/Plugins/Experimental/NNERuntimeRDG/Source/NNERuntimeRDG/Private/Hlsl/NNERuntimeRDGConvTranspose.cpp
+++ b/Engine/Plugins/Experimental/NNERuntimeRDG/Source/NNERuntimeRDG/Private/Hlsl/NNERuntimeRDGConvTranspose.cpp
@@ -1,7 +1,9 @@
 // Copyright Epic Games, Inc. All Rights Reserved.
 
 #include "NNERuntimeRDGConvTranspose.h"
+
 #include "NNEHlslShadersConvTransposeCS.h"
+#include "NNEHlslShadersLog.h"
 #include "NNERuntimeRDGHlslHelper.h"
 #include "NNETensor.h"
 #include "NNETypes.h"
@@ -74,17 +76,17 @@ namespace UE::NNERuntimeRDG::Private::Hlsl
 
 			if (Input.GetShape().Rank() < 2)
 			{
-				UE_LOG(LogNNE, Warning, TEXT("ConvTranspose first input should be at least of rank 2"));
+				UE_LOG(LogNNERuntimeRDGHlsl, Warning, TEXT("ConvTranspose: First input should be at least of rank 2"));
 				return false;
 			}
 			if (Weights.GetShape().Rank() != Input.GetShape().Rank())
 			{
-				UE_LOG(LogNNE, Warning, TEXT("ConvTranspose first and second inputs should be of same ranks"));
+				UE_LOG(LogNNERuntimeRDGHlsl, Warning, TEXT("ConvTranspose: First and second inputs should be of same ranks"));
 				return false;
 			}
 			if (Output.GetShape().Rank() != Input.GetShape().Rank())
 			{
-				UE_LOG(LogNNE, Warning, TEXT("ConvTranspose first input and output should be of same ranks"));
+				UE_LOG(LogNNERuntimeRDGHlsl, Warning, TEXT("ConvTranspose: First input and output should be of same ranks"));
 				return false;
 			}
 
diff --git a/Engine/Plugins/Experimental/NNERuntimeRDG/Source/NNERuntimeRDG/Private/Hlsl/NNERuntimeRDGCumSum.cpp b/Engine/Plugins/Experimental/NNERuntimeRDG/Source/NNERuntimeRDG/Private/Hlsl/NNERuntimeRDGCumSum.cpp
--- a/Engine/Plugins/Experimental/NNERuntimeRDG/Source/NNERuntimeRDG/Private/Hlsl/NNERuntimeRDGCumSum.cpp
+++ b/Engine/Plugins/Experimental/NNERuntimeRDG/Source/NNERuntimeRDG/Private/Hlsl/NNERuntimeRDGCumSum.cpp
@@ -1,7 +1,9 @@
 // Copyright Epic Games, Inc. All Rights Reserved.
 
 #include "NNEHlslShadersCumSumCS.h"
+
 #include "Helper/NNERuntimeRDGOperatorHelper.h"
+#include "NNEHlslShadersLog.h"
 #include "NNERuntimeRDGCumSum.h"
 #include "NNERuntimeRDGHlslHelper.h"
 #include "NNETensor.h"
@@ -39,13 +41,13 @@ namespace UE::NNERuntimeRDG::Private::Hlsl
 			const NNE::Internal::FTensorRef AxisTensor = InputTensors[1];
 			if(!AxisTensor->HasPreparedData())
 			{
-				UE_LOG(LogNNE, Warning, TEXT("CumSum tensor `axis` (name: %s) must be CPU constant."), *InputTensors[1]->GetName());
+				UE_LOG(LogNNERuntimeRDGHlsl, Warning, TEXT("CumSum: Tensor `axis` (name: %s) must be CPU constant."), *InputTensors[1]->GetName());
 				return -1;
 			}
 			
 			if(AxisTensor->GetVolume() != 1)
 			{
-				UE_LOG(LogNNE, Warning, TEXT("CumSum tensor `axis` (name: %s) must be 0-D."), *InputTensors[1]->GetName());
+				UE_LOG(LogNNERuntimeRDGHlsl, Warning, TEXT("CumSum: Tensor `axis` (name: %s) must be 0-D."), *InputTensors[1]->GetName());
 				return -1;
 			}
 
@@ -60,7 +62,7 @@ namespace UE::NNERuntimeRDG::Private::Hlsl
 
 			if (Axis > InputRank || Axis  < -InputRank)
 			{
-				UE_LOG(LogNNE, Warning, TEXT("CumSum tensor 'axis' should contain a value in the range [-r,r] with r being the rank of the input (name: %s) however got %d while rank is %d."), *InputTensors[0]->GetName(), Axis, InputRank);
+				UE_LOG(LogNNERuntimeRDGHlsl, Warning, TEXT("CumSum: Tensor 'axis' should contain a value in the range [-r,r] with r being the rank of the input (name: %s) however got %d while rank is %d."), *InputTensors[0]->GetName(), Axis, InputRank);
 				return -1;
 			}
 
@@ -80,7 +82,7 @@ namespace UE::NNERuntimeRDG::Private::Hlsl
 
 				if(OutThreadGroupCount.X > GRHIMaxDispatchThreadGroupsPerDimension.X)
 				{
-					UE_LOG(LogNNE, Warning, TEXT("CumSum input tensor (name: %s) has axis dimension greater than %d. This is not supported."), 
+					UE_LOG(LogNNERuntimeRDGHlsl, Warning, TEXT("CumSum: Input tensor (name: %s) has axis dimension greater than %d. This is not supported."), 
 						*InputTensor.GetName(), GRHIMaxDispatchThreadGroupsPerDimension.X);
 					return FIntVector{};
 				}
@@ -99,14 +101,14 @@ namespace UE::NNERuntimeRDG::Private::Hlsl
 				OutThreadGroupCount.Y = NumElemBeforeAxis;
 				if(OutThreadGroupCount.Y > GRHIMaxDispatchThreadGroupsPerDimension.Y)
 				{
-					UE_LOG(LogNNE, Warning, TEXT("CumSum input tensor (name: %s) has number of elements before axis greater than %d. This is not supported."), 
+					UE_LOG(LogNNERuntimeRDGHlsl, Warning, TEXT("CumSum: Input tensor (name: %s) has number of elements before axis greater than %d. This is not supported."), 
 						*InputTensor.GetName(), GRHIMaxDispatchThreadGroupsPerDimension.X);
 					return FIntVector{};
 				}
 				OutThreadGroupCount.Z = NumElemAfterAxis;
 				if(OutThreadGroupCount.Z > GRHIMaxDispatchThreadGroupsPerDimension.Z)
 				{
-					UE_LOG(LogNNE, Warning, TEXT("CumSum input tensor (name: %s) has number of elements after axis greater than %d. This is not supported."), 
+					UE_LOG(LogNNERuntimeRDGHlsl, Warning, TEXT("CumSum: Input tensor (name: %s) has number of elements after axis greater than %d. This is not supported."), 
 						*InputTensor.GetName(), GRHIMaxDispatchThreadGroupsPerDimension.Z);
 					return FIntVector{};
 				}
@@ -135,14 +137,14 @@ namespace UE::NNERuntimeRDG::Private::Hlsl
 			int32 Exclusive = Attributes.GetValueOrDefault<int32>(TEXT("exclusive"), 0);
 			if(Exclusive == 1)
 			{
-				UE_LOG(LogNNE, Warning, TEXT("CumSum attribute `exclusive` not yet supported."));
+				UE_LOG(LogNNERuntimeRDGHlsl, Warning, TEXT("CumSum: Attribute `exclusive` not yet supported."));
 				return false;
 			}
 
 			int32 Reverse = Attributes.GetValueOrDefault<int32>(TEXT("reverse"), 0);
 			if(Reverse == 1)
 			{
-				UE_LOG(LogNNE, Warning, TEXT("CumSum attribute `reverse` not yet supported."));
+				UE_LOG(LogNNERuntimeRDGHlsl, Warning, TEXT("CumSum: Attribute `reverse` not yet supported."));
 				return false;
 			}
 
diff --git a/Engine/Plugins/Experimental/NNERuntimeRDG/Source/NNERuntimeRDG/Private/Hlsl/NNERuntimeRDGDepthToSpace.cpp b/Engine/Plugins/Experimental/NNERuntimeRDG/Source/NNERuntimeRDG/Private/Hlsl/NNERuntimeRDGDepthToSpace.cpp
--- a/Engine/Plugins/Experimental/NNERuntimeRDG/Source/NNERuntimeRDG/Private/Hlsl/NNERuntimeRDGDepthToSpace.cpp
+++ b/Engine/Plugins/Experimental/NNERuntimeRDG/Source/NNERuntimeRDG/Private/Hlsl/NNERuntimeRDGDepthToSpace.cpp
@@ -1,7 +1,9 @@
 // Copyright Epic Games, Inc. All Rights Reserved.
 
 #include "NNERuntimeRDGDepthToSpace.h"
+
 #include "NNERuntimeRDGHlslHelper.h"
+#include "NNEHlslShadersLog.h"
 #include "NNEHlslShadersTransposeCS.h"
 #include "NNETensor.h"
 #include "NNETypes.h"
@@ -59,7 +61,7 @@ namespace UE::NNERuntimeRDG::Private::Hlsl
 
 			if(InputShape.Num() != 4)
 			{
-				UE_LOG(LogNNE, Warning, TEXT("Hlsl DepthToSpace: input tensor (name: %s) doesn't have [N,C,H,W] format."), *X.GetName());
+				UE_LOG(LogNNERuntimeRDGHlsl, Warning, TEXT("DepthToSpace: input tensor (name: %s) doesn't have [N,C,H,W] format."), *X.GetName());
 				return false;
 			}
 
@@ -98,7 +100,7 @@ namespace UE::NNERuntimeRDG::Private::Hlsl
 			const FNNEAttributeValue* BlockSizeAttr = Attributes.GetAttributeValue(TEXT("blocksize"));
 			if(BlockSizeAttr == nullptr)
 			{
-				UE_LOG(LogNNE, Warning, TEXT("Hlsl DepthToSpace: blocksize attribute is required."));
+				UE_LOG(LogNNERuntimeRDGHlsl, Warning, TEXT("DepthToSpace: blocksize attribute is required."));
 				return false;
 			}
 
@@ -106,7 +108,7 @@ namespace UE::NNERuntimeRDG::Private::Hlsl
 			Mode = ModeFromString(*Attributes.GetValueOrDefault<FString>(TEXT("mode"), TEXT("DCR")));
 			if(Mode == EMode::MAX)
 			{
-				UE_LOG(LogNNE, Warning, TEXT("Hlsl DepthToSpace: invalid mode."));
+				UE_LOG(LogNNERuntimeRDGHlsl, Warning, TEXT("DepthToSpace: invalid mode."));
 				return false;
 			}
 
diff --git a/Engine/Plugins/Experimental/NNERuntimeRDG/Source/NNERuntimeRDG/Private/Hlsl/NNERuntimeRDGDropout.cpp b/Engine/Plugins/Experimental/NNERuntimeRDG/Source/NNERuntimeRDG/Private/Hlsl/NNERuntimeRDGDropout.cpp
--- a/Engine/Plugins/Experimental/NNERuntimeRDG/Source/NNERuntimeRDG/Private/Hlsl/NNERuntimeRDGDropout.cpp
+++ b/Engine/Plugins/Experimental/NNERuntimeRDG/Source/NNERuntimeRDG/Private/Hlsl/NNERuntimeRDGDropout.cpp
@@ -1,6 +1,8 @@
 // Copyright Epic Games, Inc. All Rights Reserved.
 
 #include "NNERuntimeRDGDropout.h"
+
+#include "NNEHlslShadersLog.h"
 #include "NNETensor.h"
 #include "NNETypes.h"
 #include "RenderGraphUtils.h"
@@ -44,7 +46,7 @@ namespace UE::NNERuntimeRDG::Private::Hlsl
 
 			if (OutputTensorDescs.Num() == 2)
 			{
-				UE_LOG(LogNNE, Warning, TEXT("Dropout is only supported in inference mode at the moment, without the ability to output a 'mask'."));
+				UE_LOG(LogNNERuntimeRDGHlsl, Warning, TEXT("Dropout: Is only supported in inference mode at the moment, without the ability to output a 'mask'."));
 				return false;
 			}
 			
@@ -93,7 +95,7 @@ namespace UE::NNERuntimeRDG::Private::Hlsl
 
 		if (InputTypes.Num() > 1)
 		{
-			UE_LOG(LogNNE, Warning, TEXT("Dropout is only supported in inference mode at the moment, 'ratio' and 'training_mode' will be ignored."));
+			UE_LOG(LogNNERuntimeRDGHlsl, Warning, TEXT("Dropout: Is only supported in inference mode at the moment, 'ratio' and 'training_mode' will be ignored."));
 		}
 
 		bIsValid &= InputValidator.Validate(InputTypes);
diff --git a/Engine/Plugins/Experimental/NNERuntimeRDG/Source/NNERuntimeRDG/Private/Hlsl/NNERuntimeRDGElementWiseBinary.cpp b/Engine/Plugins/Experimental/NNERuntimeRDG/Source/NNERuntimeRDG/Private/Hlsl/NNERuntimeRDGElementWiseBinary.cpp
--- a/Engine/Plugins/Experimental/NNERuntimeRDG/Source/NNERuntimeRDG/Private/Hlsl/NNERuntimeRDGElementWiseBinary.cpp
+++ b/Engine/Plugins/Experimental/NNERuntimeRDG/Source/NNERuntimeRDG/Private/Hlsl/NNERuntimeRDGElementWiseBinary.cpp
@@ -1,8 +1,10 @@
 // Copyright Epic Games, Inc. All Rights Reserved.
 
 #include "NNERuntimeRDGElementWiseBinary.h"
-#include "NNERuntimeRDGHelperElementWiseBinary.h"
+
 #include "NNEHlslShadersElementWiseBinaryCS.h"
+#include "NNEHlslShadersLog.h"
+#include "NNERuntimeRDGHelperElementWiseBinary.h"
 #include "NNERuntimeRDGHlslHelper.h"
 #include "NNETensor.h"
 #include "NNETypes.h"
@@ -47,7 +49,7 @@ namespace UE::NNERuntimeRDG::Private::Hlsl
 				int32 RHSValue = RHSIndex >= 0 ? RHSInput.GetData()[RHSIndex] : 1;
 				if (LHSValue != RHSValue && LHSValue != 1 && RHSValue != 1)
 				{
-					UE_LOG(LogNNE, Warning, TEXT("Error while computing shape for element wise binary op, input shapes are not compatible"));
+					UE_LOG(LogNNERuntimeRDGHlsl, Warning, TEXT("ElementWiseBinary: Could not compute the shape, input shapes are not compatible"));
 					return -1;
 				}
 				int32 OutputValue = FMath::Max(LHSValue, RHSValue);
@@ -62,7 +64,7 @@ namespace UE::NNERuntimeRDG::Private::Hlsl
 
 			if (OutputTensors[0]->GetDataType() != ENNETensorDataType::Float && !OutputTensors[0]->HasPreparedData())
 			{
-				UE_LOG(LogNNE, Warning, TEXT("Error: binary element wise op output tensor could not be made constant nor it was of float type. Only floats are supported at the moment on the HLSL compute path."));
+				UE_LOG(LogNNERuntimeRDGHlsl, Warning, TEXT("ElementWiseBinary: Output tensor could not be made constant nor it was of float type. Only floats are supported at the moment."));
 				return -1;
 			}
 			
diff --git a/Engine/Plugins/Experimental/NNERuntimeRDG/Source/NNERuntimeRDG/Private/Hlsl/NNERuntimeRDGElementWiseVariadic.cpp b/Engine/Plugins/Experimental/NNERuntimeRDG/Source/NNERuntimeRDG/Private/Hlsl/NNERuntimeRDGElementWiseVariadic.cpp
--- a/Engine/Plugins/Experimental/NNERuntimeRDG/Source/NNERuntimeRDG/Private/Hlsl/NNERuntimeRDGElementWiseVariadic.cpp
+++ b/Engine/Plugins/Experimental/NNERuntimeRDG/Source/NNERuntimeRDG/Private/Hlsl/NNERuntimeRDGElementWiseVariadic.cpp
@@ -1,7 +1,9 @@
 // Copyright Epic Games, Inc. All Rights Reserved.
 
 #include "NNERuntimeRDGElementWiseVariadic.h"
+
 #include "NNEHlslShadersElementWiseVariadicCS.h"
+#include "NNEHlslShadersLog.h"
 #include "NNERuntimeRDGHlslHelper.h"
 #include "NNETensor.h"
 #include "NNETypes.h"
@@ -119,7 +121,7 @@ namespace UE::NNERuntimeRDG::Private::Hlsl
 					int32 InputValue = InputIndex >= 0 ? InputTensors[InputIdx]->GetShape().GetData()[InputIndex] : 1;
 					if (InputValue != OutputValue && InputValue != 1 && OutputValue != 1)
 					{
-						UE_LOG(LogNNE, Warning, TEXT("Error while computing shape for element wise variadic op, input shapes are not compatible"));
+						UE_LOG(LogNNERuntimeRDGHlsl, Warning, TEXT("ElementWiseVariadic: Could not compute shape for element wise variadic op, input shapes are not compatible"));
 						return -1;
 					}
 					OutputValue = FMath::Max(InputValue, OutputValue);
@@ -188,14 +190,14 @@ namespace UE::NNERuntimeRDG::Private::Hlsl
 
 		if (InputTypes.Num() == 0)
 		{
-			UE_LOG(LogNNE, Error, TEXT("Element-wise variadic operator requires at least 1 input"));
+			UE_LOG(LogNNERuntimeRDGHlsl, Error, TEXT("ElementWiseVariadic: Operator requires at least 1 input"));
 			bIsValid = false;
 		}
 		for (int32 i = 0; i < InputTypes.Num(); ++i)
 		{
 			if (InputTypes[i] != ENNETensorDataType::Float)
 			{
-				UE_LOG(LogNNE, Warning, TEXT("Element-wise variadic operator input '%d' of type '%d' is not supported, should be float at the moment."), i, int(InputTypes[i]));
+				UE_LOG(LogNNERuntimeRDGHlsl, Warning, TEXT("ElementWiseVariadic: Operator input '%d' of type '%d' is not supported, should be float at the moment."), i, int(InputTypes[i]));
 				bIsValid = false;
 			}
 		}
diff --git a/Engine/Plugins/Experimental/NNERuntimeRDG/Source/NNERuntimeRDG/Private/Hlsl/NNERuntimeRDGFlatten.cpp b/Engine/Plugins/Experimental/NNERuntimeRDG/Source/NNERuntimeRDG/Private/Hlsl/NNERuntimeRDGFlatten.cpp
--- a/Engine/Plugins/Experimental/NNERuntimeRDG/Source/NNERuntimeRDG/Private/Hlsl/NNERuntimeRDGFlatten.cpp
+++ b/Engine/Plugins/Experimental/NNERuntimeRDG/Source/NNERuntimeRDG/Private/Hlsl/NNERuntimeRDGFlatten.cpp
@@ -1,6 +1,8 @@
 // Copyright Epic Games, Inc. All Rights Reserved.
 
 #include "NNERuntimeRDGFlatten.h"
+
+#include "NNEHlslShadersLog.h"
 #include "NNETensor.h"
 #include "NNETypes.h"
 #include "RenderGraphUtils.h"
@@ -60,7 +62,7 @@ namespace UE::NNERuntimeRDG::Private::Hlsl
 
 			if (Axis > InputRank || Axis < -InputRank)
 			{
-				UE_LOG(LogNNE, Warning, TEXT("Flatten 'Axis' attribute should be in the range [-r,r] with r being the rank of the input (name: %s) however axis is %d while rank is %d."), *InputTensorDescs[0].GetName(), Axis, InputRank);
+				UE_LOG(LogNNERuntimeRDGHlsl, Warning, TEXT("Flatten: 'Axis' attribute should be in the range [-r,r] with r being the rank of the input (name: %s) however axis is %d while rank is %d."), *InputTensorDescs[0].GetName(), Axis, InputRank);
 				return false;
 			}
 			if (Axis < 0)
diff --git a/Engine/Plugins/Experimental/NNERuntimeRDG/Source/NNERuntimeRDG/Private/Hlsl/NNERuntimeRDGGather.cpp b/Engine/Plugins/Experimental/NNERuntimeRDG/Source/NNERuntimeRDG/Private/Hlsl/NNERuntimeRDGGather.cpp
--- a/Engine/Plugins/Experimental/NNERuntimeRDG/Source/NNERuntimeRDG/Private/Hlsl/NNERuntimeRDGGather.cpp
+++ b/Engine/Plugins/Experimental/NNERuntimeRDG/Source/NNERuntimeRDG/Private/Hlsl/NNERuntimeRDGGather.cpp
@@ -1,10 +1,12 @@
 // Copyright Epic Games, Inc. All Rights Reserved.
 
 #include "NNERuntimeRDGGather.h"
+
+#include "NNEAttributeMap.h"
+#include "NNEHlslShadersGatherCS.h"
+#include "NNEHlslShadersLog.h"
 #include "NNERuntimeRDGHelperGather.h"
 #include "NNERuntimeRDGHlslHelper.h"
-#include "NNEHlslShadersGatherCS.h"
-#include "NNEAttributeMap.h"
 #include "NNETypes.h"
 #include "NNETensor.h"
 
@@ -80,12 +82,12 @@ namespace UE::NNERuntimeRDG::Private::Hlsl
 			Axis = Attributes.GetValueOrDefault(TEXT("axis"), Axis);
 			if (Axis >= Data.GetShape().Rank())
 			{
-				UE_LOG(LogNNE, Warning, TEXT("Gather Axis attribute should be inferior to first input rank"));
+				UE_LOG(LogNNERuntimeRDGHlsl, Warning, TEXT("Gather: Axis attribute should be inferior to first input rank"));
 				return false;
 			}
 			if (Axis < -Data.GetShape().Rank())
 			{
-				UE_LOG(LogNNE, Warning, TEXT("Gather Axis attribute should be superior or equal to minus the first input rank"));
+				UE_LOG(LogNNERuntimeRDGHlsl, Warning, TEXT("Gather: Axis attribute should be superior or equal to minus the first input rank"));
 				return false;
 			}
 			Axis = Axis >= 0 ? Axis : Data.GetShape().Rank() + Axis;
@@ -167,14 +169,14 @@ namespace UE::NNERuntimeRDG::Private::Hlsl
 
 		if(InputShapes[0].Rank() < 1)
 		{
-			UE_LOG(LogNNE, Warning, TEXT("Hlsl Gather: input tensor must have rank >= 1."));
+			UE_LOG(LogNNERuntimeRDGHlsl, Warning, TEXT("Gather: input tensor must have rank >= 1."));
 			return false;
 		}
 
 		const int32 OutputRank = InputShapes[1].Rank() + (InputShapes[0].Rank() - 1);
 		if(OutputRank > NNEHlslShaders::Internal::FGatherConstants::MAX_NUM_DIMENSIONS)
 		{
-			UE_LOG(LogNNE, Warning, TEXT("Hlsl Gather: output tensor has rank %d higher than maximum supported: %d."), OutputRank, NNEHlslShaders::Internal::FGatherConstants::MAX_NUM_DIMENSIONS);
+			UE_LOG(LogNNERuntimeRDGHlsl, Warning, TEXT("Gather: output tensor has rank %d higher than maximum supported: %d."), OutputRank, NNEHlslShaders::Internal::FGatherConstants::MAX_NUM_DIMENSIONS);
 			return false;
 		}
 
diff --git a/Engine/Plugins/Experimental/NNERuntimeRDG/Source/NNERuntimeRDG/Private/Hlsl/NNERuntimeRDGGatherElements.cpp b/Engine/Plugins/Experimental/NNERuntimeRDG/Source/NNERuntimeRDG/Private/Hlsl/NNERuntimeRDGGatherElements.cpp
--- a/Engine/Plugins/Experimental/NNERuntimeRDG/Source/NNERuntimeRDG/Private/Hlsl/NNERuntimeRDGGatherElements.cpp
+++ b/Engine/Plugins/Experimental/NNERuntimeRDG/Source/NNERuntimeRDG/Private/Hlsl/NNERuntimeRDGGatherElements.cpp
@@ -1,9 +1,11 @@
 // Copyright Epic Games, Inc. All Rights Reserved.
 
 #include "NNERuntimeRDGGatherElements.h"
+
+#include "Helper/NNERuntimeRDGOperatorHelper.h"
 #include "NNEHlslShadersGatherElementsCS.h"
+#include "NNEHlslShadersLog.h"
 #include "NNERuntimeRDGHlslHelper.h"
-#include "Helper/NNERuntimeRDGOperatorHelper.h"
 #include "NNETensor.h"
 #include "NNETypes.h"
 
@@ -47,7 +49,7 @@ namespace UE::NNERuntimeRDG::Private::Hlsl
 			Axis = Attributes.GetValueOrDefault<int32>(TEXT("axis"), 0);
 			if (Axis < -Rank || Axis >= Rank)
 			{
-				UE_LOG(LogNNE, Warning, TEXT("'axis' attribute needs to be in the range [-Rank, Rank - 1], but it is %i with a rank of %i."), Axis, Rank);
+				UE_LOG(LogNNERuntimeRDGHlsl, Warning, TEXT("GatherElement: 'axis' attribute needs to be in the range [-Rank, Rank - 1], but it is %i with a rank of %i."), Axis, Rank);
 				return false;
 			}
 			if (Axis < 0)
diff --git a/Engine/Plugins/Experimental/NNERuntimeRDG/Source/NNERuntimeRDG/Private/Hlsl/NNERuntimeRDGGemm.cpp b/Engine/Plugins/Experimental/NNERuntimeRDG/Source/NNERuntimeRDG/Private/Hlsl/NNERuntimeRDGGemm.cpp
--- a/Engine/Plugins/Experimental/NNERuntimeRDG/Source/NNERuntimeRDG/Private/Hlsl/NNERuntimeRDGGemm.cpp
+++ b/Engine/Plugins/Experimental/NNERuntimeRDG/Source/NNERuntimeRDG/Private/Hlsl/NNERuntimeRDGGemm.cpp
@@ -1,7 +1,9 @@
 // Copyright Epic Games, Inc. All Rights Reserved.
 
 #include "NNERuntimeRDGGemm.h"
+
 #include "NNEHlslShadersGemmCS.h"
+#include "NNEHlslShadersLog.h"
 #include "NNERuntimeRDGHlslHelper.h"
 #include "NNEAttributeMap.h"
 #include "NNETensor.h"
@@ -65,12 +67,12 @@ namespace UE::NNERuntimeRDG::Private::Hlsl
 
 			if (InputA.GetShape().Rank() != 2)
 			{
-				UE_LOG(LogNNE, Warning, TEXT("Gemm first input should be of rank 2"));
+				UE_LOG(LogNNERuntimeRDGHlsl, Warning, TEXT("Gemm: First input should be of rank 2"));
 				return false;
 			}
 			if (InputB.GetShape().Rank() != 2)
 			{
-				UE_LOG(LogNNE, Warning, TEXT("Gemm second input should be of rank 2"));
+				UE_LOG(LogNNERuntimeRDGHlsl, Warning, TEXT("Gemm: Second input should be of rank 2"));
 				return false;
 			}
 			if (InputTensorDescs.Num() == 3)
@@ -78,7 +80,7 @@ namespace UE::NNERuntimeRDG::Private::Hlsl
 				const NNE::FTensorDesc& InputC = InputTensorDescs[2];
 				if (InputC.GetShape().Rank() > 2)
 				{
-					UE_LOG(LogNNE, Warning, TEXT("Gemm third input should be of rank 2 or less"));
+					UE_LOG(LogNNERuntimeRDGHlsl, Warning, TEXT("Gemm: Third input should be of rank 2 or less"));
 					return false;
 				}
 			}
diff --git a/Engine/Plugins/Experimental/NNERuntimeRDG/Source/NNERuntimeRDG/Private/Hlsl/NNERuntimeRDGGlobalPool.cpp b/Engine/Plugins/Experimental/NNERuntimeRDG/Source/NNERuntimeRDG/Private/Hlsl/NNERuntimeRDGGlobalPool.cpp
--- a/Engine/Plugins/Experimental/NNERuntimeRDG/Source/NNERuntimeRDG/Private/Hlsl/NNERuntimeRDGGlobalPool.cpp
+++ b/Engine/Plugins/Experimental/NNERuntimeRDG/Source/NNERuntimeRDG/Private/Hlsl/NNERuntimeRDGGlobalPool.cpp
@@ -1,6 +1,8 @@
 // Copyright Epic Games, Inc. All Rights Reserved.
 
 #include "NNERuntimeRDGGlobalPool.h"
+
+#include "NNEHlslShadersLog.h"
 #include "NNEHlslShadersReduceCS.h"
 #include "NNERuntimeRDGHlslHelper.h"
 #include "NNETensor.h"
@@ -57,7 +59,7 @@ namespace UE::NNERuntimeRDG::Private::Hlsl
 
 			if (InputTensorDescs[0].GetShape().Rank() != OutputTensorDescs[0].GetShape().Rank())
 			{
-				UE_LOG(LogNNE, Warning, TEXT("GlobalPool operators requires the output to have the same rank as the input."));
+				UE_LOG(LogNNERuntimeRDGHlsl, Warning, TEXT("GlobalPool: Operators requires the output to have the same rank as the input."));
 				return false;
 			}
 
@@ -65,7 +67,7 @@ namespace UE::NNERuntimeRDG::Private::Hlsl
 
 			if (InputRank <= FirstReducedDimension)
 			{
-				UE_LOG(LogNNE, Warning, TEXT("GlobalPool operators requires input tensor to be at least 3-D (but got rank %d)"), InputRank);
+				UE_LOG(LogNNERuntimeRDGHlsl, Warning, TEXT("GlobalPool: Operators requires input tensor to be at least 3-D (but got rank %d)"), InputRank);
 				return false;
 			}
 
diff --git a/Engine/Plugins/Experimental/NNERuntimeRDG/Source/NNERuntimeRDG/Private/Hlsl/NNERuntimeRDGInstanceNormalization.cpp b/Engine/Plugins/Experimental/NNERuntimeRDG/Source/NNERuntimeRDG/Private/Hlsl/NNERuntimeRDGInstanceNormalization.cpp
--- a/Engine/Plugins/Experimental/NNERuntimeRDG/Source/NNERuntimeRDG/Private/Hlsl/NNERuntimeRDGInstanceNormalization.cpp
+++ b/Engine/Plugins/Experimental/NNERuntimeRDG/Source/NNERuntimeRDG/Private/Hlsl/NNERuntimeRDGInstanceNormalization.cpp
@@ -1,7 +1,9 @@
 // Copyright Epic Games, Inc. All Rights Reserved.
 
 #include "NNERuntimeRDGInstanceNormalization.h"
+
 #include "NNEHlslShadersInstanceNormalizationCS.h"
+#include "NNEHlslShadersLog.h"
 #include "NNEHlslShadersReduceCS.h"
 #include "NNERuntimeRDGHlslHelper.h"
 #include "NNETensor.h"
@@ -14,23 +16,23 @@ int ValidateInput(const UE::NNE::FTensorShape& Input, const UE::NNE::FTensorShap
 {
 	if (Scale.Rank() != 1)
 	{
-		UE_LOG(LogNNE, Warning, TEXT("InstanceNormalization input scale should be of rank 1: %d"), Scale.Rank());
+		UE_LOG(LogNNERuntimeRDGHlsl, Warning, TEXT("InstanceNormalization: Input scale should be of rank 1: %d"), Scale.Rank());
 		return -1;
 	}
 	if (Scale.GetData()[0] != Input.GetData()[1])
 	{
-		UE_LOG(LogNNE, Warning, TEXT("InstanceNormalization input scale size should be equal to channel count: %d vs %d"), Scale.GetData()[0], Input.GetData()[1]);
+		UE_LOG(LogNNERuntimeRDGHlsl, Warning, TEXT("InstanceNormalization: Input scale size should be equal to channel count: %d vs %d"), Scale.GetData()[0], Input.GetData()[1]);
 		return -1;
 	}
 
 	if (Bias.Rank() != 1)
 	{
-		UE_LOG(LogNNE, Warning, TEXT("InstanceNormalization input B should be of rank 1: %d"), Bias.Rank());
+		UE_LOG(LogNNERuntimeRDGHlsl, Warning, TEXT("InstanceNormalization: Input B should be of rank 1: %d"), Bias.Rank());
 		return -1;
 	}
 	if (Bias.GetData()[0] != Input.GetData()[1])
 	{
-		UE_LOG(LogNNE, Warning, TEXT("InstanceNormalization intput B size should be equal to channel count: : %d vs %d"), Bias.GetData()[0], Input.GetData()[1]);
+		UE_LOG(LogNNERuntimeRDGHlsl, Warning, TEXT("InstanceNormalization: Input B size should be equal to channel count: : %d vs %d"), Bias.GetData()[0], Input.GetData()[1]);
 		return -1;
 	}
 
@@ -88,13 +90,13 @@ namespace UE::NNERuntimeRDG::Private::Hlsl
 
 			if (InputRank < 3)
 			{
-				UE_LOG(LogNNE, Warning, TEXT("InstanceNormalization input data should be at least of rank 3 but got: %d"), InputRank);
+				UE_LOG(LogNNERuntimeRDGHlsl, Warning, TEXT("InstanceNormalization input data should be at least of rank 3 but got: %d"), InputRank);
 				return false;
 			}
 
 			if (InputRank != OutputTensorDescs[0].GetShape().Rank())
 			{
-				UE_LOG(LogNNE, Warning, TEXT("InstanceNormalization requires the output to have the same rank as the input."));
+				UE_LOG(LogNNERuntimeRDGHlsl, Warning, TEXT("InstanceNormalization requires the output to have the same rank as the input."));
 				return false;
 			}
 
diff --git a/Engine/Plugins/Experimental/NNERuntimeRDG/Source/NNERuntimeRDG/Private/Hlsl/NNERuntimeRDGLayerNormalization.cpp b/Engine/Plugins/Experimental/NNERuntimeRDG/Source/NNERuntimeRDG/Private/Hlsl/NNERuntimeRDGLayerNormalization.cpp
--- a/Engine/Plugins/Experimental/NNERuntimeRDG/Source/NNERuntimeRDG/Private/Hlsl/NNERuntimeRDGLayerNormalization.cpp
+++ b/Engine/Plugins/Experimental/NNERuntimeRDG/Source/NNERuntimeRDG/Private/Hlsl/NNERuntimeRDGLayerNormalization.cpp
@@ -1,7 +1,9 @@
 // Copyright Epic Games, Inc. All Rights Reserved.
 
 #include "NNERuntimeRDGLayerNormalization.h"
+
 #include "NNEHlslShadersLayerNormalizationCS.h"
+#include "NNEHlslShadersLog.h"
 #include "NNEHlslShadersReduceCS.h"
 #include "NNERuntimeRDGHlslHelper.h"
 #include "NNETensor.h"
@@ -38,14 +40,14 @@ namespace UE::NNERuntimeRDG::Private::Hlsl
 			{
 				if (ScaleOrBias.Rank() > (Input.Rank() - Axis))
 				{
-					UE_LOG(LogNNE, Warning, TEXT("LayerNormalization: Scale/bias tensor rank is invalid: %d"), ScaleOrBias.Rank());
+					UE_LOG(LogNNERuntimeRDGHlsl, Warning, TEXT("LayerNormalization: Scale/bias tensor rank is invalid: %d"), ScaleOrBias.Rank());
 					return -1;
 				}
 				for(int DimIdx = 0; DimIdx < ScaleOrBias.Rank(); ++DimIdx)
 				{
 					if(Input.GetData()[DimIdx + Axis] != ScaleOrBias.GetData()[DimIdx] && ScaleOrBias.GetData()[DimIdx] != 1)
 					{
-						UE_LOG(LogNNE, Warning, TEXT("LayerNormalization: Scale/bias tensor shape not equal nor broadcastable to input's"));
+						UE_LOG(LogNNERuntimeRDGHlsl, Warning, TEXT("LayerNormalization: Scale/bias tensor shape not equal nor broadcastable to input's"));
 						return -1;
 					}
 				}
@@ -113,7 +115,7 @@ namespace UE::NNERuntimeRDG::Private::Hlsl
 
 			if (InputRank != OutputTensorDescs[0].GetShape().Rank())
 			{
-				UE_LOG(LogNNE, Warning, TEXT("LayerNormalization requires the output to have the same rank as the input."));
+				UE_LOG(LogNNERuntimeRDGHlsl, Warning, TEXT("LayerNormalization: Output should have the same rank as the input."));
 				return false;
 			}
 
diff --git a/Engine/Plugins/Experimental/NNERuntimeRDG/Source/NNERuntimeRDG/Private/Hlsl/NNERuntimeRDGMatMul.cpp b/Engine/Plugins/Experimental/NNERuntimeRDG/Source/NNERuntimeRDG/Private/Hlsl/NNERuntimeRDGMatMul.cpp
--- a/Engine/Plugins/Experimental/NNERuntimeRDG/Source/NNERuntimeRDG/Private/Hlsl/NNERuntimeRDGMatMul.cpp
+++ b/Engine/Plugins/Experimental/NNERuntimeRDG/Source/NNERuntimeRDG/Private/Hlsl/NNERuntimeRDGMatMul.cpp
@@ -1,7 +1,9 @@
 // Copyright Epic Games, Inc. All Rights Reserved.
 
 #include "NNERuntimeRDGMatMul.h"
+
 #include "NNEHlslShadersGemmCS.h"
+#include "NNEHlslShadersLog.h"
 #include "NNERuntimeRDGHlslHelper.h"
 #include "NNETensor.h"
 #include "NNETypes.h"
@@ -32,17 +34,17 @@ namespace UE::NNERuntimeRDG::Private::Hlsl
 
 			if (InputA.Rank() < 2)
 			{
-				UE_LOG(LogNNE, Warning, TEXT("Matmul first input should be at least of rank 2"));
+				UE_LOG(LogNNERuntimeRDGHlsl, Warning, TEXT("Matmul: First input should be at least of rank 2"));
 				return -1;
 			}
 			if (InputB.Rank() < 2)
 			{
-				UE_LOG(LogNNE, Warning, TEXT("Matmul second input should be at least of rank 2"));
+				UE_LOG(LogNNERuntimeRDGHlsl, Warning, TEXT("Matmul: Second input should be at least of rank 2"));
 				return -1;
 			}
 			if (InputA.GetData()[InputA.Rank() - 1] != InputB.GetData()[InputB.Rank() - 2])
 			{
-				UE_LOG(LogNNE, Warning, TEXT("Matmul first input last dimension should be equal to second input last dimension"));
+				UE_LOG(LogNNERuntimeRDGHlsl, Warning, TEXT("Matmul: First input last dimension should be equal to second input last dimension"));
 				return -1;
 			}
 
@@ -82,12 +84,12 @@ namespace UE::NNERuntimeRDG::Private::Hlsl
 
 			if (InputA.GetShape().Rank() < 2)
 			{
-				UE_LOG(LogNNE, Warning, TEXT("Matmul first input should be at least of rank 2"));
+				UE_LOG(LogNNERuntimeRDGHlsl, Warning, TEXT("Matmul: First input should be at least of rank 2"));
 				return false;
 			}
 			if (InputB.GetShape().Rank() < 2)
 			{
-				UE_LOG(LogNNE, Warning, TEXT("Matmul second input should be at least of rank 2"));
+				UE_LOG(LogNNERuntimeRDGHlsl, Warning, TEXT("Matmul: Second input should be at least of rank 2"));
 				return false;
 			}
 
diff --git a/Engine/Plugins/Experimental/NNERuntimeRDG/Source/NNERuntimeRDG/Private/Hlsl/NNERuntimeRDGPad.cpp b/Engine/Plugins/Experimental/NNERuntimeRDG/Source/NNERuntimeRDG/Private/Hlsl/NNERuntimeRDGPad.cpp
--- a/Engine/Plugins/Experimental/NNERuntimeRDG/Source/NNERuntimeRDG/Private/Hlsl/NNERuntimeRDGPad.cpp
+++ b/Engine/Plugins/Experimental/NNERuntimeRDG/Source/NNERuntimeRDG/Private/Hlsl/NNERuntimeRDGPad.cpp
@@ -1,6 +1,8 @@
 // Copyright Epic Games, Inc. All Rights Reserved.
 
 #include "NNERuntimeRDGPad.h"
+
+#include "NNEHlslShadersLog.h"
 #include "NNEHlslShadersPadCS.h"
 #include "NNERuntimeRDGHlslHelper.h"
 #include "Helper/NNERuntimeRDGOperatorHelper.h"
@@ -43,7 +45,7 @@ namespace UE::NNERuntimeRDG::Private::Hlsl
 				const NNE::Internal::FTensorRef PadsTensor = InputTensors[1];
 				if(!PadsTensor->IsConstant())
 				{
-					UE_LOG(LogNNE, Warning, TEXT("Pad input 'pads' (name: %s) should be constant."), *PadsTensor->GetName());
+					UE_LOG(LogNNERuntimeRDGHlsl, Warning, TEXT("Pad: Input 'pads' (name: %s) should be constant."), *PadsTensor->GetName());
 					return -1;
 				}
 
@@ -53,13 +55,13 @@ namespace UE::NNERuntimeRDG::Private::Hlsl
 
 					if(!ValueTensor.HasPreparedData())
 					{
-						UE_LOG(LogNNE, Warning, TEXT("Pad input 'constant_value' (name: %s) should be constant."), *ValueTensor.GetName());
+						UE_LOG(LogNNERuntimeRDGHlsl, Warning, TEXT("Pad: Input 'constant_value' (name: %s) should be constant."), *ValueTensor.GetName());
 						return -1;
 					}
 
 					if(ValueTensor.GetPreparedData<float>().Num() != 1)
 					{
-						UE_LOG(LogNNE, Warning, TEXT("Pad input 'constant_value' (name: %s) should be scalar, however it is not."), *ValueTensor.GetName());
+						UE_LOG(LogNNERuntimeRDGHlsl, Warning, TEXT("Pad: Input 'constant_value' (name: %s) should be scalar, however it is not."), *ValueTensor.GetName());
 						return -1;
 					}
 
@@ -72,7 +74,7 @@ namespace UE::NNERuntimeRDG::Private::Hlsl
 
 					if(!AxesTensor->IsConstant())
 					{
-						UE_LOG(LogNNE, Warning, TEXT("Pad input 'axes' (name: %s) should be constant."), *AxesTensor->GetName());
+						UE_LOG(LogNNERuntimeRDGHlsl, Warning, TEXT("Pad: Input 'axes' (name: %s) should be constant."), *AxesTensor->GetName());
 						return -1;
 					}
 
@@ -90,7 +92,7 @@ namespace UE::NNERuntimeRDG::Private::Hlsl
 
 					if (RelativePads.Num() != Axes.Num() * 2)
 					{
-						UE_LOG(LogNNE, Warning, TEXT("Pad input 'axes' (name: %s) has to have a size that is twice the size \
+						UE_LOG(LogNNERuntimeRDGHlsl, Warning, TEXT("Pad: Input 'axes' (name: %s) has to have a size that is twice the size \
 													  of pad input 'pads' (name: %s), but they have size %i and %i respectively."), 
 												*AxesTensor->GetName(), *PadsTensor->GetName(), Axes.Num(), RelativePads.Num());
 						return -1;
@@ -102,7 +104,7 @@ namespace UE::NNERuntimeRDG::Private::Hlsl
 						int32 Axis = Axes[AxesIndex];
 						if (Axis < -Rank || Axis >= Rank)
 						{
-							UE_LOG(LogNNE, Warning, TEXT("Pad input value at index %i of the 'axes' (name: %s) tensor \
+							UE_LOG(LogNNERuntimeRDGHlsl, Warning, TEXT("Pad: Input value at index %i of the 'axes' (name: %s) tensor \
 														  needs to be in the range [-Rank, Rank - 1], but value is %i with a rank of %i."), 
 													AxesIndex, *AxesTensor->GetName(), Axis, Rank);
 							return -1;
@@ -119,7 +121,7 @@ namespace UE::NNERuntimeRDG::Private::Hlsl
 				{
 					if (!PadsTensor->HasPreparedData())
 					{
-						UE_LOG(LogNNE, Warning, TEXT("pads attribute lenght (%d) should be twice the rank of input X (%d)."), 0, X.GetShape().Rank());
+						UE_LOG(LogNNERuntimeRDGHlsl, Warning, TEXT("Pad: 'pads' attribute lenght (%d) should be twice the rank of input X (%d)."), 0, X.GetShape().Rank());
 						return -1;
 					}
 					OperatorHelper::GetInt32ArrayFromConstTensor(Pads, PadsTensor);
@@ -128,7 +130,7 @@ namespace UE::NNERuntimeRDG::Private::Hlsl
 
 			if ((2*X.GetShape().Rank()) != Pads.Num())
 			{
-				UE_LOG(LogNNE, Warning, TEXT("pads attribute lenght (%d) should be twice the rank of input X (%d)."), Pads.Num(), X.GetShape().Rank());
+				UE_LOG(LogNNERuntimeRDGHlsl, Warning, TEXT("Pad: 'pads' attribute lenght (%d) should be twice the rank of input X (%d)."), Pads.Num(), X.GetShape().Rank());
 				return false;
 			}
 
@@ -140,7 +142,7 @@ namespace UE::NNERuntimeRDG::Private::Hlsl
 				int32 OutputDim = PrePad + X.GetShape().GetData()[i] + PostPad;
 				if (OutputDim < 1)
 				{
-					UE_LOG(LogNNE, Warning, TEXT("Pads cannot reduce dimension below 1, but would for tensor (name:%s) at rank %d of size %d with prepad %d and postpad %d."), *X.GetName(), i, X.GetShape().GetData()[i], PrePad, PostPad);
+					UE_LOG(LogNNERuntimeRDGHlsl, Warning, TEXT("Pad: Cannot reduce dimension below 1, but would for tensor (name:%s) at rank %d of size %d with prepad %d and postpad %d."), *X.GetName(), i, X.GetShape().GetData()[i], PrePad, PostPad);
 					return -1;
 				}
 				OutputShapeData.Emplace(OutputDim);
diff --git a/Engine/Plugins/Experimental/NNERuntimeRDG/Source/NNERuntimeRDG/Private/Hlsl/NNERuntimeRDGPool.cpp b/Engine/Plugins/Experimental/NNERuntimeRDG/Source/NNERuntimeRDG/Private/Hlsl/NNERuntimeRDGPool.cpp
--- a/Engine/Plugins/Experimental/NNERuntimeRDG/Source/NNERuntimeRDG/Private/Hlsl/NNERuntimeRDGPool.cpp
+++ b/Engine/Plugins/Experimental/NNERuntimeRDG/Source/NNERuntimeRDG/Private/Hlsl/NNERuntimeRDGPool.cpp
@@ -1,7 +1,9 @@
 // Copyright Epic Games, Inc. All Rights Reserved.
 
 #include "NNERuntimeRDGPool.h"
+
 #include "NNEHlslShadersConvCS.h"
+#include "NNEHlslShadersLog.h"
 #include "NNEHlslShadersPoolCS.h"
 #include "NNEHlslShadersTypeHelper.h"
 #include "NNERuntimeRDGHlslHelper.h"
@@ -103,7 +105,7 @@ namespace UE::NNERuntimeRDG::Private::Hlsl
 				check(OutputTensorDescs.Num() >= 1);
 				if (OutputTensorDescs.Num() > 1)
 				{
-					UE_LOG(LogNNE, Warning, TEXT("MaxPool 2nd optional output 'Indices' is not supported."));
+					UE_LOG(LogNNERuntimeRDGHlsl, Warning, TEXT("MaxPool: 2nd optional output 'Indices' is not supported."));
 					return false;
 				}
 			}
@@ -117,7 +119,7 @@ namespace UE::NNERuntimeRDG::Private::Hlsl
 
 			if (Input.GetShape().Rank() < 3)
 			{
-				UE_LOG(LogNNE, Warning, TEXT("%s input should be at least of rank 3, to have 1+ spatial dimension(s) but is of rank %d"), GetOperatorName(), Input.GetShape().Rank());
+				UE_LOG(LogNNERuntimeRDGHlsl, Warning, TEXT("%s : Input should be at least of rank 3, to have 1+ spatial dimension(s) but is of rank %d"), GetOperatorName(), Input.GetShape().Rank());
 				return false;
 			}
 			NumSpatialDimensions = Input.GetShape().Rank() - 2;
@@ -139,22 +141,22 @@ namespace UE::NNERuntimeRDG::Private::Hlsl
 
 			if (KernelShape.Num() != NumSpatialDimensions)
 			{
-				UE_LOG(LogNNE, Warning, TEXT("%s KernelShape should have as many elements as the spatial dimensions of the input, got %d while input have %d."), GetOperatorName(), KernelShape.Num(), NumSpatialDimensions);
+				UE_LOG(LogNNERuntimeRDGHlsl, Warning, TEXT("%s: KernelShape should have as many elements as the spatial dimensions of the input, got %d while input have %d."), GetOperatorName(), KernelShape.Num(), NumSpatialDimensions);
 				return false;
 			}
 			if (Strides.Num() != NumSpatialDimensions)
 			{
-				UE_LOG(LogNNE, Warning, TEXT("%s Strides should have as many elements as the spatial dimensions of the input, got %d while input have %d."), GetOperatorName(), Strides.Num(), NumSpatialDimensions);
+				UE_LOG(LogNNERuntimeRDGHlsl, Warning, TEXT("%s: Strides should have as many elements as the spatial dimensions of the input, got %d while input have %d."), GetOperatorName(), Strides.Num(), NumSpatialDimensions);
 				return false;
 			}
 			if (Dilations.Num() != NumSpatialDimensions)
 			{
-				UE_LOG(LogNNE, Warning, TEXT("%s Dilations should have as many elements as the spatial dimensions of the input, got %d while input have %d."), GetOperatorName(), Dilations.Num(), NumSpatialDimensions);
+				UE_LOG(LogNNERuntimeRDGHlsl, Warning, TEXT("%s: Dilations should have as many elements as the spatial dimensions of the input, got %d while input have %d."), GetOperatorName(), Dilations.Num(), NumSpatialDimensions);
 				return false;
 			}
 			if (Pads.Num() != 2*NumSpatialDimensions)
 			{
-				UE_LOG(LogNNE, Warning, TEXT("%s Pads should have twice as many elements as the spatial dimensions of the input, got %d while input have %d."), GetOperatorName(), Pads.Num(), NumSpatialDimensions);
+				UE_LOG(LogNNERuntimeRDGHlsl, Warning, TEXT("%s: Pads should have twice as many elements as the spatial dimensions of the input, got %d while input have %d."), GetOperatorName(), Pads.Num(), NumSpatialDimensions);
 				return false;
 			}
 
diff --git a/Engine/Plugins/Experimental/NNERuntimeRDG/Source/NNERuntimeRDG/Private/Hlsl/NNERuntimeRDGReduce.cpp b/Engine/Plugins/Experimental/NNERuntimeRDG/Source/NNERuntimeRDG/Private/Hlsl/NNERuntimeRDGReduce.cpp
--- a/Engine/Plugins/Experimental/NNERuntimeRDG/Source/NNERuntimeRDG/Private/Hlsl/NNERuntimeRDGReduce.cpp
+++ b/Engine/Plugins/Experimental/NNERuntimeRDG/Source/NNERuntimeRDG/Private/Hlsl/NNERuntimeRDGReduce.cpp
@@ -1,13 +1,15 @@
 // Copyright Epic Games, Inc. All Rights Reserved.
 
 #include "NNERuntimeRDGReduce.h"
+
+#include "Algo/Sort.h"
+#include "Helper/NNERuntimeRDGOperatorHelper.h"
+#include "Misc/EnumerateRange.h"
+#include "NNEHlslShadersLog.h"
 #include "NNEHlslShadersReduceCS.h"
 #include "NNERuntimeRDGHlslHelper.h"
-#include "Helper/NNERuntimeRDGOperatorHelper.h"
 #include "NNETensor.h"
 #include "NNETypes.h"
-#include "Algo/Sort.h"
-#include "Misc/EnumerateRange.h"
 #include "RenderGraphUtils.h"
 
 namespace UE::NNERuntimeRDG::Private::Hlsl
@@ -50,7 +52,7 @@ namespace UE::NNERuntimeRDG::Private::Hlsl
 					const NNE::Internal::FTensorRef AxesTensor = InputTensors[1];
 					if(!AxesTensor->HasPreparedData())
 					{
-						UE_LOG(LogNNE, Warning, TEXT("Reduce tensor `axes` (name: %s) must be CPU constant."), *InputTensors[1]->GetName());
+						UE_LOG(LogNNERuntimeRDGHlsl, Warning, TEXT("Reduce: Tensor `axes` (name: %s) must be CPU constant."), *InputTensors[1]->GetName());
 						return false;
 					}
 					if(AxesTensor->GetVolume() != 0)
@@ -67,7 +69,7 @@ namespace UE::NNERuntimeRDG::Private::Hlsl
 			{
 				if (Axis > InputRank || Axis  < -InputRank)
 				{
-					UE_LOG(LogNNE, Warning, TEXT("Reduce operators 'Axes' attribute should contain value be in the range [-r,r] with r being the rank of the input (name: %s) however got %d while rank is %d."), *InputTensors[0]->GetName(), Axis, InputRank);
+					UE_LOG(LogNNERuntimeRDGHlsl, Warning, TEXT("Reduce: 'Axes' attribute should contain value be in the range [-r,r] with r being the rank of the input (name: %s) however got %d while rank is %d."), *InputTensors[0]->GetName(), Axis, InputRank);
 					return false;
 				}
 
@@ -128,7 +130,7 @@ namespace UE::NNERuntimeRDG::Private::Hlsl
 
 			if (!bAxesAsInput && Axes.Num() == 0)
 			{
-				UE_LOG(LogNNE, Warning, TEXT("Reduce attribute `axes` cannot be empty."));
+				UE_LOG(LogNNERuntimeRDGHlsl, Warning, TEXT("Reduce: Attribute `axes` cannot be empty."));
 				return false;
 			}
 
diff --git a/Engine/Plugins/Experimental/NNERuntimeRDG/Source/NNERuntimeRDG/Private/Hlsl/NNERuntimeRDGReshape.cpp b/Engine/Plugins/Experimental/NNERuntimeRDG/Source/NNERuntimeRDG/Private/Hlsl/NNERuntimeRDGReshape.cpp
--- a/Engine/Plugins/Experimental/NNERuntimeRDG/Source/NNERuntimeRDG/Private/Hlsl/NNERuntimeRDGReshape.cpp
+++ b/Engine/Plugins/Experimental/NNERuntimeRDG/Source/NNERuntimeRDG/Private/Hlsl/NNERuntimeRDGReshape.cpp
@@ -1,10 +1,12 @@
 // Copyright Epic Games, Inc. All Rights Reserved.
 
 #include "NNERuntimeRDGReshape.h"
+
+#include "Helper/NNERuntimeRDGHelperReshape.h"
+#include "NNEHlslShadersLog.h"
 #include "NNETensor.h"
 #include "NNETypes.h"
 #include "RenderGraphUtils.h"
-#include "Helper/NNERuntimeRDGHelperReshape.h"
 
 namespace UE::NNERuntimeRDG::Private::Hlsl
 {
@@ -37,7 +39,7 @@ namespace UE::NNERuntimeRDG::Private::Hlsl
 
 			if (!ShapeTensor.HasPreparedData())
 			{
-				UE_LOG(LogNNE, Warning, TEXT("Reshape input 'Shape' (name: %s) should be constant for shape inference to succeed, however it is not."), *ShapeTensor.GetName());
+				UE_LOG(LogNNERuntimeRDGHlsl, Warning, TEXT("Reshape: Input 'Shape' (name: %s) should be constant for shape inference to succeed, however it is not."), *ShapeTensor.GetName());
 				return -1;
 			}
 
diff --git a/Engine/Plugins/Experimental/NNERuntimeRDG/Source/NNERuntimeRDG/Private/Hlsl/NNERuntimeRDGResize.cpp b/Engine/Plugins/Experimental/NNERuntimeRDG/Source/NNERuntimeRDG/Private/Hlsl/NNERuntimeRDGResize.cpp
--- a/Engine/Plugins/Experimental/NNERuntimeRDG/Source/NNERuntimeRDG/Private/Hlsl/NNERuntimeRDGResize.cpp
+++ b/Engine/Plugins/Experimental/NNERuntimeRDG/Source/NNERuntimeRDG/Private/Hlsl/NNERuntimeRDGResize.cpp
@@ -1,15 +1,16 @@
 // Copyright Epic Games, Inc. All Rights Reserved.
 
 #include "NNERuntimeRDGResize.h"
-#include "NNEHlslShadersResizeCS.h"
 
+#include "Algo/MaxElement.h"
+#include "Algo/MinElement.h"
+#include "NNEHlslShadersLog.h"
+#include "NNEHlslShadersResizeCS.h"
 #include "NNEHlslShadersTypeHelper.h"
 #include "NNERuntimeRDGHlslHelper.h"
 #include "NNETensor.h"
 #include "NNETypes.h"
 #include "RenderGraphUtils.h"
-#include "Algo/MinElement.h"
-#include "Algo/MaxElement.h"
 
 namespace UE::NNERuntimeRDG::Private::Hlsl
 {
@@ -73,7 +74,7 @@ namespace UE::NNERuntimeRDG::Private::Hlsl
 			{
 				if(!Roi.HasPreparedData())
 				{
-					UE_LOG(LogNNE, Warning, TEXT("Hlsl Resize: `roi` tensor could not be made constant. (name %s)."), *Roi.GetName());
+					UE_LOG(LogNNERuntimeRDGHlsl, Warning, TEXT("Resize: `roi` tensor could not be made constant. (name %s)."), *Roi.GetName());
 					return -1;
 				}
 				
@@ -81,7 +82,7 @@ namespace UE::NNERuntimeRDG::Private::Hlsl
 
 				if(RegionOfInterest.Num() != 2 * Input.GetShape().Rank())
 				{
-					UE_LOG(LogNNE, Warning, TEXT("Hlsl Resize: `roi` tensor (name %s) must have 2 * N length."), *Roi.GetName());
+					UE_LOG(LogNNERuntimeRDGHlsl, Warning, TEXT("Resize: `roi` tensor (name %s) must have 2 * N length."), *Roi.GetName());
 					return -1;
 				}
 			}
@@ -96,7 +97,7 @@ namespace UE::NNERuntimeRDG::Private::Hlsl
 				const NNE::Internal::FTensor& Scales = *InputTensors[2];
 				if(!Scales.HasPreparedData())
 				{
-					UE_LOG(LogNNE, Warning, TEXT("Hlsl Resize: `scales` tensor could not be made constant. (name %s)."), *Scales.GetName());
+					UE_LOG(LogNNERuntimeRDGHlsl, Warning, TEXT("Resize: `scales` tensor could not be made constant. (name %s)."), *Scales.GetName());
 					return -1;
 				}
 
@@ -128,7 +129,7 @@ namespace UE::NNERuntimeRDG::Private::Hlsl
 				const NNE::Internal::FTensor& Sizes = *InputTensors[3];
 				if(!Sizes.HasPreparedData())
 				{
-					UE_LOG(LogNNE, Warning, TEXT("Hlsl Resize: `sizes` tensor could not be made constant. (name %s)."), *Sizes.GetName());
+					UE_LOG(LogNNERuntimeRDGHlsl, Warning, TEXT("Resize: `sizes` tensor could not be made constant. (name %s)."), *Sizes.GetName());
 					return -1;
 				}
 
@@ -177,7 +178,7 @@ namespace UE::NNERuntimeRDG::Private::Hlsl
 			// NOTE: implementation for antialiasing is missing
 			if(Attributes.GetValueOrDefault<int32>(TEXT("antialias"), 0) == 1)
 			{
-				UE_LOG(LogNNE, Warning, TEXT("Hlsl Resize: antialias not yet supported."));
+				UE_LOG(LogNNERuntimeRDGHlsl, Warning, TEXT("Resize: antialias not yet supported."));
 				return false;
 			}
 
@@ -187,7 +188,7 @@ namespace UE::NNERuntimeRDG::Private::Hlsl
 			{
 				if (Axis > InputTensorDescs[0].GetShape().Rank() || Axis < -InputTensorDescs[0].GetShape().Rank())
 				{
-					UE_LOG(LogNNE, Warning, TEXT("Hlsl Resize: 'Axes' attribute should contain value be in the range [-r,r] with r being the rank of the input (name: %s) however got %d while rank is %d."), *InputTensorDescs[0].GetName(), Axis, InputTensorDescs[0].GetShape().Rank());
+					UE_LOG(LogNNERuntimeRDGHlsl, Warning, TEXT("Resize: 'Axes' attribute should contain value be in the range [-r,r] with r being the rank of the input (name: %s) however got %d while rank is %d."), *InputTensorDescs[0].GetName(), Axis, InputTensorDescs[0].GetShape().Rank());
 					return false;
 				}
 
@@ -200,7 +201,7 @@ namespace UE::NNERuntimeRDG::Private::Hlsl
 			// NOTE: `axes` attribute not yet supported
 			if(Axes.Num() != 0)
 			{
-				UE_LOG(LogNNE, Warning, TEXT("Hlsl Resize: `axes` attribute not yet supported."));
+				UE_LOG(LogNNERuntimeRDGHlsl, Warning, TEXT("Resize: `axes` attribute not yet supported."));
 				return false;
 			}
 
@@ -210,7 +211,7 @@ namespace UE::NNERuntimeRDG::Private::Hlsl
 			// NOTE: `exclude_outside` attribute not yet supported
 			if(ExcludeOutside == 1)
 			{
-				UE_LOG(LogNNE, Warning, TEXT("Hlsl Resize: `exclude_outside` attribute not yet supported."));
+				UE_LOG(LogNNERuntimeRDGHlsl, Warning, TEXT("Resize: `exclude_outside` attribute not yet supported."));
 				return false;
 			}
 
@@ -220,7 +221,7 @@ namespace UE::NNERuntimeRDG::Private::Hlsl
 			//NOTE: cubic interpolation not yet supported
 			if(Mode == UE::NNEHlslShaders::Internal::EMode::Cubic)
 			{
-				UE_LOG(LogNNE, Warning, TEXT("Hlsl Resize: cubic interpolation not yet supported."));
+				UE_LOG(LogNNERuntimeRDGHlsl, Warning, TEXT("Resize: Cubic interpolation not yet supported."));
 				return false;
 			}
 
@@ -344,13 +345,13 @@ namespace UE::NNERuntimeRDG::Private::Hlsl
 
 		if(InputTypes.Num() < 3 || InputTypes.Num() > 4)
 		{
-			UE_LOG(LogNNE, Warning, TEXT("Got a total of '%d' inputs but should be between 3 and 4."), InputTypes.Num());
+			UE_LOG(LogNNERuntimeRDGHlsl, Warning, TEXT("Resize: Got a total of '%d' inputs but should be between 3 and 4."), InputTypes.Num());
 			return false;
 		}
 
 		if(InputShapes[0].Rank() < 1)
 		{
-			UE_LOG(LogNNE, Warning, TEXT("Hlsl Resize: input tensor must have rank >= 1."));
+			UE_LOG(LogNNERuntimeRDGHlsl, Warning, TEXT("Resize: Input tensor must have rank >= 1."));
 			return false;
 		}
 
@@ -358,14 +359,14 @@ namespace UE::NNERuntimeRDG::Private::Hlsl
 		{
 			if(InputShapes[1].Rank() != 1)
 			{
-				UE_LOG(LogNNE, Warning, TEXT("Roi tensor must be a 1-D tensor."));
+				UE_LOG(LogNNERuntimeRDGHlsl, Warning, TEXT("Resize: Roi tensor must be a 1-D tensor."));
 				return false;
 			}
 			FString TransMode = AttributeMap.GetValueOrDefault<FString>(TEXT("coordinate_transformation_mode"), TEXT("half_pixel"));
 			bool IsCropAndResizeMode = FResizeCS::CoordTransModeFromString(*TransMode) == ECoordTransMode::TfCropAndResize;
 			if(IsCropAndResizeMode && InputShapes[1].GetData()[0] != 2 * InputShapes[0].Rank())
 			{
-				UE_LOG(LogNNE, Warning, TEXT("Roi tensor must have dimension 2*N (where N is the input rank) when `coordinate_transformation_mode` is `tf_crop_and_resize`."));
+				UE_LOG(LogNNERuntimeRDGHlsl, Warning, TEXT("Resize: Roi tensor must have dimension 2*N (where N is the input rank) when `coordinate_transformation_mode` is `tf_crop_and_resize`."));
 				return false;
 			}
 		}
@@ -374,12 +375,12 @@ namespace UE::NNERuntimeRDG::Private::Hlsl
 		{
 			if(InputShapes[2].Rank() != 1)
 			{
-				UE_LOG(LogNNE, Warning, TEXT("Scales tensor must be a 1-D tensor."));
+				UE_LOG(LogNNERuntimeRDGHlsl, Warning, TEXT("Resize: Scales tensor must be a 1-D tensor."));
 				return false;
 			}
 			if(InputShapes[2].GetData()[0] != InputShapes[0].Rank())
 			{
-				UE_LOG(LogNNE, Warning, TEXT("Scales tensor must have dimension N (where N is the input rank)."));
+				UE_LOG(LogNNERuntimeRDGHlsl, Warning, TEXT("Resize: Scales tensor must have dimension N (where N is the input rank)."));
 				return false;
 			}
 		}
@@ -387,17 +388,17 @@ namespace UE::NNERuntimeRDG::Private::Hlsl
 		{
 			if(InputTypes[2] != ENNETensorDataType::None)
 			{
-				UE_LOG(LogNNE, Warning, TEXT("Scales tensor must be empty (i.e. empty name and data type 'None') when Sizes is specified."));
+				UE_LOG(LogNNERuntimeRDGHlsl, Warning, TEXT("Resize: Scales tensor must be empty (i.e. empty name and data type 'None') when Sizes is specified."));
 				return false;
 			}
 			if(InputShapes[3].Rank() != 1)
 			{
-				UE_LOG(LogNNE, Warning, TEXT("Sizes tensor must be a 1-D tensor."));
+				UE_LOG(LogNNERuntimeRDGHlsl, Warning, TEXT("Resize: Sizes tensor must be a 1-D tensor."));
 				return false;
 			}
 			if(InputShapes[3].GetData()[0] != InputShapes[0].Rank())
 			{
-				UE_LOG(LogNNE, Warning, TEXT("Sizes tensor must have dimension N (where N is the input rank)."));
+				UE_LOG(LogNNERuntimeRDGHlsl, Warning, TEXT("Resize: Sizes tensor must have dimension N (where N is the input rank)."));
 				return false;
 			}
 		}
diff --git a/Engine/Plugins/Experimental/NNERuntimeRDG/Source/NNERuntimeRDG/Private/Hlsl/NNERuntimeRDGScatterND.cpp b/Engine/Plugins/Experimental/NNERuntimeRDG/Source/NNERuntimeRDG/Private/Hlsl/NNERuntimeRDGScatterND.cpp
--- a/Engine/Plugins/Experimental/NNERuntimeRDG/Source/NNERuntimeRDG/Private/Hlsl/NNERuntimeRDGScatterND.cpp
+++ b/Engine/Plugins/Experimental/NNERuntimeRDG/Source/NNERuntimeRDG/Private/Hlsl/NNERuntimeRDGScatterND.cpp
@@ -1,12 +1,14 @@
 // Copyright Epic Games, Inc. All Rights Reserved.
 
 #include "NNERuntimeRDGScatterND.h"
+
+#include "Algo/Compare.h"
+#include "NNEHlslShadersLog.h"
 #include "NNEHlslShadersScatterNDCS.h"
 #include "NNERuntimeRDGHlslHelper.h"
 #include "NNETensor.h"
 #include "NNETypes.h"
 #include "RenderGraphUtils.h"
-#include "Algo/Compare.h"
 
 namespace UE::NNERuntimeRDG::Private::Hlsl
 {
@@ -41,28 +43,28 @@ namespace UE::NNERuntimeRDG::Private::Hlsl
 
 			if (IndicesShape.Last() > (uint32) InputShape.Num())
 			{
-				UE_LOG(LogNNE, Warning, TEXT("Last dimension in the shape of `indices` (%d) must be less than input rank (%d)."), IndicesShape.Last(), InputShape.Num());
+				UE_LOG(LogNNERuntimeRDGHlsl, Warning, TEXT("ScatterND: Last dimension in the shape of `indices` (%d) must be less than input rank (%d)."), IndicesShape.Last(), InputShape.Num());
 				return -1;
 			}
 			if (UpdatesShape.Num() != (IndicesShape.Num() - 1) + (InputShape.Num() - IndicesShape.Last()))
 			{
-				UE_LOG(LogNNE, Warning, TEXT("Rank of `updates` should equal (q - 1) + (r - k), with q rank of `indices`, r rank of `data` and `k` last dimension of `indices`' shape."));
+				UE_LOG(LogNNERuntimeRDGHlsl, Warning, TEXT("ScatterND: Rank of `updates` should equal (q - 1) + (r - k), with q rank of `indices`, r rank of `data` and `k` last dimension of `indices`' shape."));
 				return -1;
 			}
 			if (!Algo::Compare( UpdatesShape.Slice(0, IndicesShape.Num() - 1), IndicesShape.Slice(0, IndicesShape.Num() - 1) ))
 			{
-				UE_LOG(LogNNE, Warning, TEXT("updates.shape[0:q-1] should match indices.shape[0:q-1]."));
+				UE_LOG(LogNNERuntimeRDGHlsl, Warning, TEXT("ScatterND: Updates.shape[0:q-1] should match indices.shape[0:q-1]."));
 				return -1;
 			}
 			if (!Algo::Compare( UpdatesShape.RightChop(IndicesShape.Num() - 1), InputShape.RightChop(IndicesShape.Last()) ))
 			{
-				UE_LOG(LogNNE, Warning, TEXT("updates.shape[q-1:] should match data.shape[k:]."));
+				UE_LOG(LogNNERuntimeRDGHlsl, Warning, TEXT("ScatterND: Updates.shape[q-1:] should match data.shape[k:]."));
 				return -1;
 			}
 			//NOTE: we need to limit the volume of Input to the int32 maximum expressible value due to unavailability of int64 type in shaders.
 			if (Input.GetShape().Volume() > (uint64) MAX_int32)
 			{
-				UE_LOG(LogNNE, Warning, TEXT("This implementation of ScatterND only supports input tensors up to a volume of %u."), MAX_int32);
+				UE_LOG(LogNNERuntimeRDGHlsl, Warning, TEXT("ScatterND: Only supports input tensors up to a volume of %u."), MAX_int32);
 				return -1;
 			}
 
diff --git a/Engine/Plugins/Experimental/NNERuntimeRDG/Source/NNERuntimeRDG/Private/Hlsl/NNERuntimeRDGShape.cpp b/Engine/Plugins/Experimental/NNERuntimeRDG/Source/NNERuntimeRDG/Private/Hlsl/NNERuntimeRDGShape.cpp
--- a/Engine/Plugins/Experimental/NNERuntimeRDG/Source/NNERuntimeRDG/Private/Hlsl/NNERuntimeRDGShape.cpp
+++ b/Engine/Plugins/Experimental/NNERuntimeRDG/Source/NNERuntimeRDG/Private/Hlsl/NNERuntimeRDGShape.cpp
@@ -1,6 +1,8 @@
 // Copyright Epic Games, Inc. All Rights Reserved.
 
 #include "NNERuntimeRDGShape.h"
+
+#include "NNEHlslShadersLog.h"
 #include "NNETensor.h"
 #include "NNETypes.h"
 
@@ -50,7 +52,7 @@ namespace UE::NNERuntimeRDG::Private::Hlsl
 
 			if (OutputTensorDescs[0].GetDataType() != ENNETensorDataType::Int64)
 			{
-				UE_LOG(LogNNE, Warning, TEXT("Shape should output a tensor of type Int64"));
+				UE_LOG(LogNNERuntimeRDGHlsl, Warning, TEXT("Shape: Should output a tensor of type Int64"));
 				return false;
 			}
 			
@@ -59,7 +61,7 @@ namespace UE::NNERuntimeRDG::Private::Hlsl
 
 		virtual void Dispatch(FRDGBuilder& GraphBuilder, TConstArrayView<FTensorRDGRef> InputTensors, TConstArrayView<FTensorRDGRef> OutputTensors) override
 		{
-			UE_LOG(LogNNE, Warning, TEXT("Shape: Output should be constant and already uploaded to GPU memory. Dispatch should not need to be called."));
+			UE_LOG(LogNNERuntimeRDGHlsl, Warning, TEXT("Shape: Output should be constant and already uploaded to GPU memory. Dispatch should not need to be called."));
 		}
 	};
 
diff --git a/Engine/Plugins/Experimental/NNERuntimeRDG/Source/NNERuntimeRDG/Private/Hlsl/NNERuntimeRDGSize.cpp b/Engine/Plugins/Experimental/NNERuntimeRDG/Source/NNERuntimeRDG/Private/Hlsl/NNERuntimeRDGSize.cpp
--- a/Engine/Plugins/Experimental/NNERuntimeRDG/Source/NNERuntimeRDG/Private/Hlsl/NNERuntimeRDGSize.cpp
+++ b/Engine/Plugins/Experimental/NNERuntimeRDG/Source/NNERuntimeRDG/Private/Hlsl/NNERuntimeRDGSize.cpp
@@ -1,6 +1,8 @@
 // Copyright Epic Games, Inc. All Rights Reserved.
 
 #include "NNERuntimeRDGSize.h"
+
+#include "NNEHlslShadersLog.h"
 #include "NNETensor.h"
 #include "NNETypes.h"
 
@@ -47,7 +49,7 @@ namespace UE::NNERuntimeRDG::Private::Hlsl
 
 		virtual void Dispatch(FRDGBuilder& GraphBuilder, TConstArrayView<FTensorRDGRef> InputTensors, TConstArrayView<FTensorRDGRef> OutputTensors) override
 		{
-			UE_LOG(LogNNE, Warning, TEXT("Size: Output should be constant and already uploaded to GPU memory. Dispatch should not need to be called."));
+			UE_LOG(LogNNERuntimeRDGHlsl, Warning, TEXT("Size: Output should be constant and already uploaded to GPU memory. Dispatch should not need to be called."));
 		}
 	};
 
diff --git a/Engine/Plugins/Experimental/NNERuntimeRDG/Source/NNERuntimeRDG/Private/Hlsl/NNERuntimeRDGSlice.cpp b/Engine/Plugins/Experimental/NNERuntimeRDG/Source/NNERuntimeRDG/Private/Hlsl/NNERuntimeRDGSlice.cpp
--- a/Engine/Plugins/Experimental/NNERuntimeRDG/Source/NNERuntimeRDG/Private/Hlsl/NNERuntimeRDGSlice.cpp
+++ b/Engine/Plugins/Experimental/NNERuntimeRDG/Source/NNERuntimeRDG/Private/Hlsl/NNERuntimeRDGSlice.cpp
@@ -1,7 +1,9 @@
 // Copyright Epic Games, Inc. All Rights Reserved.
 
 #include "NNERuntimeRDGSlice.h"
+
 #include "Helper/NNERuntimeRDGOperatorHelper.h"
+#include "NNEHlslShadersLog.h"
 #include "NNEHlslShadersSliceCS.h"
 #include "NNERuntimeRDGHelperSlice.h"
 #include "NNERuntimeRDGHlslHelper.h"
@@ -43,12 +45,12 @@ namespace UE::NNERuntimeRDG::Private::Hlsl
 			{
 				if (!OperatorHelper::GetInt32ArrayFromConstTensor(AxesAttr, InputTensors[3]))
 				{
-					UE_LOG(LogNNE, Warning, TEXT("Error: Slice op 'Axes' input tensor (%s) is only supported as a constant integer tensor but it is not."), *InputTensors[3]->GetName());
+					UE_LOG(LogNNERuntimeRDGHlsl, Warning, TEXT("Slice: 'Axes' input tensor (%s) is only supported as a constant integer tensor but it is not."), *InputTensors[3]->GetName());
 					return false;
 				}
 				if(AxesAttr.Num() > InputTensors[0]->GetShape().Rank() || AxesAttr.Num() < 1)
 				{
-					UE_LOG(LogNNE, Warning, TEXT("Error: Slice op 'Axes' input tensor (%s) contains %d elements but input rank is %d."), *InputTensors[3]->GetName(), AxesAttr.Num(), InputTensors[0]->GetShape().Rank());
+					UE_LOG(LogNNERuntimeRDGHlsl, Warning, TEXT("Slice: 'Axes' input tensor (%s) contains %d elements but input rank is %d."), *InputTensors[3]->GetName(), AxesAttr.Num(), InputTensors[0]->GetShape().Rank());
 					return false;
 				}
 			}
@@ -64,23 +66,23 @@ namespace UE::NNERuntimeRDG::Private::Hlsl
 
 			if (!OperatorHelper::GetInt32ArrayFromConstTensor(StartsAttr, InputTensors[1]))
 			{
-				UE_LOG(LogNNE, Warning, TEXT("Error: Slice op 'Starts' input tensor (%s) is only supported as a constant integer tensor but it is not."), *InputTensors[1]->GetName());
+				UE_LOG(LogNNERuntimeRDGHlsl, Warning, TEXT("Slice: 'Starts' input tensor (%s) is only supported as a constant integer tensor but it is not."), *InputTensors[1]->GetName());
 				return false;
 			}
 			if(StartsAttr.Num() != NumAxes)
 			{
-				UE_LOG(LogNNE, Warning, TEXT("Error: Slice op 'Starts' input tensor (%s) contains %d elements but number of axes is %d."), *InputTensors[1]->GetName(), StartsAttr.Num(), NumAxes);
+				UE_LOG(LogNNERuntimeRDGHlsl, Warning, TEXT("Slice: 'Starts' input tensor (%s) contains %d elements but number of axes is %d."), *InputTensors[1]->GetName(), StartsAttr.Num(), NumAxes);
 				return false;
 			}
 
 			if (!OperatorHelper::GetInt32ArrayFromConstTensor(EndsAttr, InputTensors[2]))
 			{
-				UE_LOG(LogNNE, Warning, TEXT("Error: Slice op 'Ends' input tensor (%s) is only supported as a constant integer tensor but it is not."), *InputTensors[2]->GetName());
+				UE_LOG(LogNNERuntimeRDGHlsl, Warning, TEXT("Slice: 'Ends' input tensor (%s) is only supported as a constant integer tensor but it is not."), *InputTensors[2]->GetName());
 				return false;
 			}
 			if(EndsAttr.Num() != NumAxes)
 			{
-				UE_LOG(LogNNE, Warning, TEXT("Error: Slice op 'Ends' input tensor (%s) contains %d elements but number of axes is %d."), *InputTensors[2]->GetName(), EndsAttr.Num(), NumAxes);
+				UE_LOG(LogNNERuntimeRDGHlsl, Warning, TEXT("Slice: 'Ends' input tensor (%s) contains %d elements but number of axes is %d."), *InputTensors[2]->GetName(), EndsAttr.Num(), NumAxes);
 				return false;
 			}
 
@@ -90,19 +92,19 @@ namespace UE::NNERuntimeRDG::Private::Hlsl
 			{
 				if (!OperatorHelper::GetInt32ArrayFromConstTensor(StepsAttr, InputTensors[4]))
 				{
-					UE_LOG(LogNNE, Warning, TEXT("Error: Slice op 'Steps' input tensor (%s) is only supported as a constant integer tensor but it is not."), *InputTensors[4]->GetName());
+					UE_LOG(LogNNERuntimeRDGHlsl, Warning, TEXT("Slice: 'Steps' input tensor (%s) is only supported as a constant integer tensor but it is not."), *InputTensors[4]->GetName());
 					return false;
 				}
 				if(StepsAttr.Num() != NumAxes)
 				{
-					UE_LOG(LogNNE, Warning, TEXT("Error: Slice op 'Steps' input tensor (%s) contains %d elements but number of axes is %d."), *InputTensors[4]->GetName(), StepsAttr.Num(), NumAxes);
+					UE_LOG(LogNNERuntimeRDGHlsl, Warning, TEXT("Slice: 'Steps' input tensor (%s) contains %d elements but number of axes is %d."), *InputTensors[4]->GetName(), StepsAttr.Num(), NumAxes);
 					return false;
 				}
 				for(const int32 Value : StepsAttr)
 				{
 					if(Value == 0)
 					{
-						UE_LOG(LogNNE, Warning, TEXT("Error: Slice op 'Steps' tensor (%s) can only contain non-0 integers."), *InputTensors[4]->GetName());
+						UE_LOG(LogNNERuntimeRDGHlsl, Warning, TEXT("Slice: 'Steps' tensor (%s) can only contain non-0 integers."), *InputTensors[4]->GetName());
 						return false;
 					}
 				}
@@ -208,7 +210,7 @@ namespace UE::NNERuntimeRDG::Private::Hlsl
 				int32 RangeSize = Step[Idx] > 0 ? End[Idx] - Start[Idx] : Start[Idx] - End[Idx];
 				if(RangeSize < 1)
 				{
-					UE_LOG(LogNNE, Warning, TEXT("Error: Slice op start (tensor: %s) and end (tensor: %s) indices are incompatible with step direction for dimension %d."), *InputTensors[1]->GetName(), *InputTensors[2]->GetName(), Idx);
+					UE_LOG(LogNNERuntimeRDGHlsl, Warning, TEXT("Slice: Start (tensor: %s) and end (tensor: %s) indices are incompatible with step direction for dimension %d."), *InputTensors[1]->GetName(), *InputTensors[2]->GetName(), Idx);
 					return -1;
 				}
 				uint32 OutDimSize = FMath::DivideAndRoundUp(RangeSize, FMath::Abs(Step[Idx]));
@@ -223,7 +225,7 @@ namespace UE::NNERuntimeRDG::Private::Hlsl
 
 			if (InputTensors[0]->HasPreparedData() && !OutputTensors[0]->HasPreparedData())
 			{
-				UE_LOG(LogNNE, Warning, TEXT("Error: Slice op output tensor (%s) could not be constant-folded from input."), *OutputTensors[0]->GetName());
+				UE_LOG(LogNNERuntimeRDGHlsl, Warning, TEXT("Slice: Output tensor (%s) could not be constant-folded from input."), *OutputTensors[0]->GetName());
 				return -1;
 			}
 
@@ -252,7 +254,7 @@ namespace UE::NNERuntimeRDG::Private::Hlsl
 
 				if (EndsAttr.Num() != StartsAttr.Num() || AxesAttr.Num() != StartsAttr.Num())
 				{
-					UE_LOG(LogNNE, Warning, TEXT("Slice: Starts, Ends and Axes must be of the same size."));
+					UE_LOG(LogNNERuntimeRDGHlsl, Warning, TEXT("Slice: Starts, Ends and Axes must be of the same size."));
 					return false;
 				}
 
diff --git a/Engine/Plugins/Experimental/NNERuntimeRDG/Source/NNERuntimeRDG/Private/Hlsl/NNERuntimeRDGSoftmax.cpp b/Engine/Plugins/Experimental/NNERuntimeRDG/Source/NNERuntimeRDG/Private/Hlsl/NNERuntimeRDGSoftmax.cpp
--- a/Engine/Plugins/Experimental/NNERuntimeRDG/Source/NNERuntimeRDG/Private/Hlsl/NNERuntimeRDGSoftmax.cpp
+++ b/Engine/Plugins/Experimental/NNERuntimeRDG/Source/NNERuntimeRDG/Private/Hlsl/NNERuntimeRDGSoftmax.cpp
@@ -1,6 +1,8 @@
 // Copyright Epic Games, Inc. All Rights Reserved.
 
 #include "NNERuntimeRDGSoftmax.h"
+
+#include "NNEHlslShadersLog.h"
 #include "NNEHlslShadersReduceCS.h"
 #include "NNEHlslShadersSoftmaxCS.h"
 #include "NNERuntimeRDGHlslHelper.h"
@@ -51,13 +53,13 @@ namespace UE::NNERuntimeRDG::Private::Hlsl
 
 			if (Version <= 11 && InputDimensions < 2)
 			{
-				UE_LOG(LogNNE, Warning, TEXT("Input tensor should be at least 2-D (but got rank %d)"), InputDimensions);
+				UE_LOG(LogNNERuntimeRDGHlsl, Warning, TEXT("Softmax: Input tensor should be at least 2-D (but got rank %d)"), InputDimensions);
 				return false;
 			}
 
 			if (InputTensorDescs[0].GetShape().Rank() != OutputTensorDescs[0].GetShape().Rank())
 			{
-				UE_LOG(LogNNE, Warning, TEXT("Softwax requires the output to have the same rank as the input."));
+				UE_LOG(LogNNERuntimeRDGHlsl, Warning, TEXT("Softwax: Output should have the same rank as the input."));
 				return false;
 			}
 
@@ -67,7 +69,7 @@ namespace UE::NNERuntimeRDG::Private::Hlsl
 			Axis = Attributes.GetValueOrDefault<int32>(TEXT("axis"), Version <= 11 ? 1 : -1);
 			if (Axis < AxisLow || AxisHigh < Axis)
 			{
-				UE_LOG(LogNNE, Warning, TEXT("Invalid axis (should in the interval [%d, %d], but got %d)"), AxisLow, AxisHigh, Axis);
+				UE_LOG(LogNNERuntimeRDGHlsl, Warning, TEXT("Softmax: Invalid axis (should in the interval [%d, %d], but got %d)"), AxisLow, AxisHigh, Axis);
 				return false;
 			}
 			if(Axis < 0)
diff --git a/Engine/Plugins/Experimental/NNERuntimeRDG/Source/NNERuntimeRDG/Private/Hlsl/NNERuntimeRDGSplit.cpp b/Engine/Plugins/Experimental/NNERuntimeRDG/Source/NNERuntimeRDG/Private/Hlsl/NNERuntimeRDGSplit.cpp
--- a/Engine/Plugins/Experimental/NNERuntimeRDG/Source/NNERuntimeRDG/Private/Hlsl/NNERuntimeRDGSplit.cpp
+++ b/Engine/Plugins/Experimental/NNERuntimeRDG/Source/NNERuntimeRDG/Private/Hlsl/NNERuntimeRDGSplit.cpp
@@ -1,6 +1,8 @@
 // Copyright Epic Games, Inc. All Rights Reserved.
 
 #include "NNERuntimeRDGSplit.h"
+
+#include "NNEHlslShadersLog.h"
 #include "NNEHlslShadersSplitCS.h"
 #include "NNERuntimeRDGHlslHelper.h"
 #include "NNETensor.h"
@@ -49,7 +51,7 @@ namespace UE::NNERuntimeRDG::Private::Hlsl
 				const NNE::Internal::FTensor& SplitTensor = *InputTensors[1];
 				if(!SplitTensor.HasPreparedData())
 				{
-					UE_LOG(LogNNE, Warning, TEXT("Split: RDG Hlsl requires the split tensor to be CPU-constant in order to compute output shapes."));
+					UE_LOG(LogNNERuntimeRDGHlsl, Warning, TEXT("Split: split tensor need to be CPU-constant in order to compute output shapes."));
 					return -1;
 				}
 				
@@ -65,19 +67,19 @@ namespace UE::NNERuntimeRDG::Private::Hlsl
 
 			if(!CheckSplitsSum(InputShape))
 			{
-				UE_LOG(LogNNE, Warning, TEXT("Split: Sum of split values not equal to split axis' dimension."));
+				UE_LOG(LogNNERuntimeRDGHlsl, Warning, TEXT("Split: Sum of split values not equal to split axis' dimension."));
 				return -1;
 			}
 
 			if(Splits.Num() > FSplitConstants::MAX_NUM_SPLITS)
 			{
-				UE_LOG(LogNNE, Warning, TEXT("Split: Number of splits (%d) exceeds maximum allowed (%d)."), Splits.Num(), FSplitConstants::MAX_NUM_SPLITS);
+				UE_LOG(LogNNERuntimeRDGHlsl, Warning, TEXT("Split: Number of splits (%d) exceeds maximum allowed (%d)."), Splits.Num(), FSplitConstants::MAX_NUM_SPLITS);
 				return -1;
 			}
 			
 			if(OutputTensors.Num() != Splits.Num())
 			{
-				UE_LOG(LogNNE, Warning, TEXT("Split: Number of output tensors differs from number of splits provided."));
+				UE_LOG(LogNNERuntimeRDGHlsl, Warning, TEXT("Split: Number of output tensors differs from number of splits provided."));
 				return -1;
 			}
 
@@ -100,7 +102,7 @@ namespace UE::NNERuntimeRDG::Private::Hlsl
 			Axis = Attributes.GetValueOrDefault<int32>(TEXT("axis"), Axis);
 			if (Axis > InputRank || Axis  < -InputRank)
 			{
-				UE_LOG(LogNNE, Warning, TEXT("Split's attribute 'Axis' should be in the range [-r,r] with r being the rank of the input (name: %s) however got %d while rank is %d."), *InputTensorDescs[0].GetName(), Axis, InputRank);
+				UE_LOG(LogNNERuntimeRDGHlsl, Warning, TEXT("Split: Attribute 'Axis' should be in the range [-r,r] with r being the rank of the input (name: %s) however got %d while rank is %d."), *InputTensorDescs[0].GetName(), Axis, InputRank);
 				return false;
 			}
 			// Canonicalize Axis
@@ -119,7 +121,7 @@ namespace UE::NNERuntimeRDG::Private::Hlsl
 				const int32 NumOutputs = Attributes.GetValue<int32>(TEXT("num_outputs"));
 				if(OutputTensorDescs.Num() != NumOutputs)
 				{
-					UE_LOG(LogNNE, Warning, TEXT("Split's attribute 'num_outputs' doesn't match number of output tensors. Value: %d. Number of output tensors: %d."), NumOutputs, OutputTensorDescs.Num());
+					UE_LOG(LogNNERuntimeRDGHlsl, Warning, TEXT("Split: Attribute 'num_outputs' doesn't match number of output tensors. Value: %d. Number of output tensors: %d."), NumOutputs, OutputTensorDescs.Num());
 					return false;
 				}
 			}
diff --git a/Engine/Plugins/Experimental/NNERuntimeRDG/Source/NNERuntimeRDG/Private/Hlsl/NNERuntimeRDGSqueeze.cpp b/Engine/Plugins/Experimental/NNERuntimeRDG/Source/NNERuntimeRDG/Private/Hlsl/NNERuntimeRDGSqueeze.cpp
--- a/Engine/Plugins/Experimental/NNERuntimeRDG/Source/NNERuntimeRDG/Private/Hlsl/NNERuntimeRDGSqueeze.cpp
+++ b/Engine/Plugins/Experimental/NNERuntimeRDG/Source/NNERuntimeRDG/Private/Hlsl/NNERuntimeRDGSqueeze.cpp
@@ -1,6 +1,8 @@
 // Copyright Epic Games, Inc. All Rights Reserved.
 
 #include "NNERuntimeRDGSqueeze.h"
+
+#include "NNEHlslShadersLog.h"
 #include "NNETensor.h"
 #include "NNETypes.h"
 #include "RenderGraphUtils.h"
@@ -37,7 +39,7 @@ namespace UE::NNERuntimeRDG::Private::Hlsl
 
 				if (OutputShapeData[Axe] != 1)
 				{
-					UE_LOG(LogNNE, Warning, TEXT("Squeeze at axe %d for 'Data' (name: %s) should be targeting a dimension of size 1 but it is %d."), Axe, *X.GetName(), OutputShapeData[Axe]);
+					UE_LOG(LogNNERuntimeRDGHlsl, Warning, TEXT("Squeeze: Axe %d for 'Data' (name: %s) should be targeting a dimension of size 1 but it is %d."), Axe, *X.GetName(), OutputShapeData[Axe]);
 					return -1;
 				}
 
@@ -65,12 +67,12 @@ namespace UE::NNERuntimeRDG::Private::Hlsl
 			{
 				if (Axe < 0)
 				{
-					UE_LOG(LogNNE, Warning, TEXT("Squeeze operator does not support negative axes"));
+					UE_LOG(LogNNERuntimeRDGHlsl, Warning, TEXT("Squeeze: Does not support negative axes"));
 					return false;
 				}
 				if (Axe >= InputTensorDescs[0].GetShape().Rank())
 				{
-					UE_LOG(LogNNE, Warning, TEXT("Squeeze operator does not support axes greater than the number of dimensions of the input tensor shape"));
+					UE_LOG(LogNNERuntimeRDGHlsl, Warning, TEXT("Squeeze: Does not support axes greater than the number of dimensions of the input tensor shape"));
 					return false;
 				}
 			}
diff --git a/Engine/Plugins/Experimental/NNERuntimeRDG/Source/NNERuntimeRDG/Private/Hlsl/NNERuntimeRDGTranspose.cpp b/Engine/Plugins/Experimental/NNERuntimeRDG/Source/NNERuntimeRDG/Private/Hlsl/NNERuntimeRDGTranspose.cpp
--- a/Engine/Plugins/Experimental/NNERuntimeRDG/Source/NNERuntimeRDG/Private/Hlsl/NNERuntimeRDGTranspose.cpp
+++ b/Engine/Plugins/Experimental/NNERuntimeRDG/Source/NNERuntimeRDG/Private/Hlsl/NNERuntimeRDGTranspose.cpp
@@ -1,6 +1,8 @@
 // Copyright Epic Games, Inc. All Rights Reserved.
 
 #include "NNERuntimeRDGTranspose.h"
+
+#include "NNEHlslShadersLog.h"
 #include "NNEHlslShadersTransposeCS.h"
 #include "NNERuntimeRDGHelperTranspose.h"
 #include "NNERuntimeRDGHlslHelper.h"
@@ -66,7 +68,7 @@ namespace UE::NNERuntimeRDG::Private::Hlsl
 
 			if (Perm.Num() != InputTensorDescs[0].GetShape().Rank())
 			{
-				UE_LOG(LogNNE, Warning, TEXT("Transpose 'perm' attribute should contain the same amount of element as the rank of the input tensor."));
+				UE_LOG(LogNNERuntimeRDGHlsl, Warning, TEXT("Transpose: 'perm' attribute should contain the same amount of element as the rank of the input tensor."));
 				return false;
 			}
 
diff --git a/Engine/Plugins/Experimental/NNERuntimeRDG/Source/NNERuntimeRDG/Private/Hlsl/NNERuntimeRDGUnsqueeze.cpp b/Engine/Plugins/Experimental/NNERuntimeRDG/Source/NNERuntimeRDG/Private/Hlsl/NNERuntimeRDGUnsqueeze.cpp
--- a/Engine/Plugins/Experimental/NNERuntimeRDG/Source/NNERuntimeRDG/Private/Hlsl/NNERuntimeRDGUnsqueeze.cpp
+++ b/Engine/Plugins/Experimental/NNERuntimeRDG/Source/NNERuntimeRDG/Private/Hlsl/NNERuntimeRDGUnsqueeze.cpp
@@ -1,6 +1,8 @@
 // Copyright Epic Games, Inc. All Rights Reserved.
 
 #include "NNERuntimeRDGUnsqueeze.h"
+
+#include "NNEHlslShadersLog.h"
 #include "NNETensor.h"
 #include "NNETypes.h"
 #include "RenderGraphUtils.h"
@@ -57,12 +59,12 @@ namespace UE::NNERuntimeRDG::Private::Hlsl
 			{
 				if (Axe < 0)
 				{
-					UE_LOG(LogNNE, Warning, TEXT("Unsqueeze operator does not support negative axes"));
+					UE_LOG(LogNNERuntimeRDGHlsl, Warning, TEXT("Unsqueeze: Does not support negative axes"));
 					return false;
 				}
 				if (Axe >= InputTensorDescs[0].GetShape().Rank() + Axes.Num())
 				{
-					UE_LOG(LogNNE, Warning, TEXT("Unsqueeze operator does not support axes greater than the number of dimensions of the resulting tensor shape"));
+					UE_LOG(LogNNERuntimeRDGHlsl, Warning, TEXT("Unsqueeze: Does not support axes greater than the number of dimensions of the resulting tensor shape"));
 					return false;
 				}
 			}
diff --git a/Engine/Plugins/Experimental/NNERuntimeRDG/Source/NNERuntimeRDG/Private/Hlsl/NNERuntimeRDGUpsample.cpp b/Engine/Plugins/Experimental/NNERuntimeRDG/Source/NNERuntimeRDG/Private/Hlsl/NNERuntimeRDGUpsample.cpp
--- a/Engine/Plugins/Experimental/NNERuntimeRDG/Source/NNERuntimeRDG/Private/Hlsl/NNERuntimeRDGUpsample.cpp
+++ b/Engine/Plugins/Experimental/NNERuntimeRDG/Source/NNERuntimeRDG/Private/Hlsl/NNERuntimeRDGUpsample.cpp
@@ -1,6 +1,8 @@
 // Copyright Epic Games, Inc. All Rights Reserved.
 
 #include "NNERuntimeRDGUpsample.h"
+
+#include "NNEHlslShadersLog.h"
 #include "NNERuntimeRDGAttributes.h"
 #include "NNEHlslShadersUpsampleCS.h"
 #include "NNERuntimeRDGHlslHelper.h"
@@ -37,7 +39,7 @@ namespace UE::NNERuntimeRDG::Private::Hlsl
 			
 			if (!Scales.HasPreparedData())
 			{
-				UE_LOG(LogNNE, Warning, TEXT("Upsample input 'Scale' (name: %s) should be constant for shape inference to succeed, however it is not constant."), *Scales.GetName());
+				UE_LOG(LogNNERuntimeRDGHlsl, Warning, TEXT("Upsample: Input 'Scale' (name: %s) should be constant for shape inference to succeed, however it is not constant."), *Scales.GetName());
 				return -1;
 			}
 
@@ -45,7 +47,7 @@ namespace UE::NNERuntimeRDG::Private::Hlsl
 
 			if (ScalesData.Num() != Input.GetShape().Rank())
 			{
-				UE_LOG(LogNNE, Warning, TEXT("Upsample input 'Scale' (name: %s) have %d elements. While it should be the same as the rank of input 'X' (name : %s) witch is %d"), *Scales.GetName(), ScalesData.Num(), *Input.GetName(), Input.GetShape().Rank());
+				UE_LOG(LogNNERuntimeRDGHlsl, Warning, TEXT("Upsample: Input 'Scale' (name: %s) have %d elements. While it should be the same as the rank of input 'X' (name : %s) witch is %d"), *Scales.GetName(), ScalesData.Num(), *Input.GetName(), Input.GetShape().Rank());
 				return -1;
 			}
 
@@ -53,13 +55,13 @@ namespace UE::NNERuntimeRDG::Private::Hlsl
 			{
 				if (ScalesData[i] < 1.0f)
 				{
-					UE_LOG(LogNNE, Warning, TEXT("Upsample input 'Scale' takes values greater than or equal to 1."));
+					UE_LOG(LogNNERuntimeRDGHlsl, Warning, TEXT("Upsample: Input 'Scale' takes values greater than or equal to 1."));
 					return -1;
 				}
 
 				if (Mode.Equals(AttrValue::Linear) && i < ScalesData.Num() - 3 && ScalesData[i] > 1.0f)
 				{
-					UE_LOG(LogNNE, Warning, TEXT("Upsample in 'Linear mode' only support up to trilinear interpolation, meaning input 'Scale' values for the outermost dimensions (Rank - 3) are 1."));
+					UE_LOG(LogNNERuntimeRDGHlsl, Warning, TEXT("Upsample: 'Linear mode' only support up to trilinear interpolation, meaning input 'Scale' values for the outermost dimensions (Rank - 3) are 1."));
 					return -1;
 				}
 			}
@@ -156,7 +158,7 @@ namespace UE::NNERuntimeRDG::Private::Hlsl
 		FString Mode = AttributeMap.GetValueOrDefault<FString>(AttrName::Mode, AttrValue::Nearest);
 		if (!Mode.Equals(AttrValue::Nearest) && !Mode.Equals(AttrValue::Linear))
 		{
-			UE_LOG(LogNNE, Warning, TEXT("Upsample HLSL operator only supports %s or %s as value for attribute %s"), AttrValue::Nearest, AttrValue::Linear, AttrName::Mode);
+			UE_LOG(LogNNERuntimeRDGHlsl, Warning, TEXT("Upsample: Only supports %s or %s as value for attribute %s"), AttrValue::Nearest, AttrValue::Linear, AttrName::Mode);
 			return false;
 		}
 
diff --git a/Engine/Plugins/Experimental/NNERuntimeRDG/Source/NNERuntimeRDG/Private/NNERuntimeRDGBase.cpp b/Engine/Plugins/Experimental/NNERuntimeRDG/Source/NNERuntimeRDG/Private/NNERuntimeRDGBase.cpp
--- a/Engine/Plugins/Experimental/NNERuntimeRDG/Source/NNERuntimeRDG/Private/NNERuntimeRDGBase.cpp
+++ b/Engine/Plugins/Experimental/NNERuntimeRDG/Source/NNERuntimeRDG/Private/NNERuntimeRDGBase.cpp
@@ -42,12 +42,12 @@ bool FInputValidator::Validate(TConstArrayView<ENNETensorDataType> InputTypes)
 	
 	if (InputTypes.Num() < NumRequiredInput)
 	{
-		UE_LOG(LogNNE, Warning, TEXT("Required '%d' inputs but found '%d'."), NumRequiredInput, InputTypes.Num());
+		UE_LOG(LogNNERuntimeRDGHlsl, Warning, TEXT("Required '%d' inputs but found '%d'."), NumRequiredInput, InputTypes.Num());
 		bAreInputValid = false;
 	}
 	if (InputTypes.Num() > NumRequiredInput+NumOptionalInput)
 	{
-		UE_LOG(LogNNE, Warning, TEXT("Got a total of '%d' inputs but should have '%d' maximum."), InputTypes.Num(), NumRequiredInput + NumOptionalInput);
+		UE_LOG(LogNNERuntimeRDGHlsl, Warning, TEXT("Got a total of '%d' inputs but should have '%d' maximum."), InputTypes.Num(), NumRequiredInput + NumOptionalInput);
 		bAreInputValid = false;
 	}
 	
@@ -60,7 +60,7 @@ bool FInputValidator::Validate(TConstArrayView<ENNETensorDataType> InputTypes)
 			!(InputTypes[Idx] == ENNETensorDataType::None && Idx >= NumRequiredInput))
 		{
 			FString TargetType = LogHelper::GetTensorDataTypeName(InputTypes[Idx]);
-			UE_LOG(LogNNE, Warning, TEXT("Input at index '%d' (from template T%d) is of type '%s' witch is not supported for that input."), Idx, TemplateIdx, *TargetType);
+			UE_LOG(LogNNERuntimeRDGHlsl, Warning, TEXT("Input at index '%d' (from template T%d) is of type '%s' witch is not supported for that input."), Idx, TemplateIdx, *TargetType);
 			bAreInputValid = false;
 		}
 	}
@@ -126,13 +126,13 @@ bool FAttributeValidator::Validate(const NNE::FAttributeMap& AttributesToValidat
 		if (FoundAttribute == nullptr)
 		{
 			bAreAttributesValid = false;
-			UE_LOG(LogNNE, Warning, TEXT("Required attribute '%s' not found."),
+			UE_LOG(LogNNERuntimeRDGHlsl, Warning, TEXT("Required attribute '%s' not found."),
 				*RequiredAttributes[Idx].Name);
 		}
 		else if (RequiredAttributes[Idx].Type != FoundAttribute->GetType())
 		{
 			bAreAttributesValid = false;
-			UE_LOG(LogNNE, Warning, TEXT("Required attribute '%s' type '%d' does not match expected type '%d'."),
+			UE_LOG(LogNNERuntimeRDGHlsl, Warning, TEXT("Required attribute '%s' type '%d' does not match expected type '%d'."),
 				*RequiredAttributes[Idx].Name,
 				(int)FoundAttribute->GetType(),
 				(int)RequiredAttributes[Idx].Type);
@@ -147,7 +147,7 @@ bool FAttributeValidator::Validate(const NNE::FAttributeMap& AttributesToValidat
 		if ((FoundAttribute != nullptr) && (OptionalAttributes[Idx].Type != FoundAttribute->GetType()))
 		{
 			bAreAttributesValid = false;
-			UE_LOG(LogNNE, Warning, TEXT("Optional attribute '%s' type '%d' does not match expected type '%d'."),
+			UE_LOG(LogNNERuntimeRDGHlsl, Warning, TEXT("Optional attribute '%s' type '%d' does not match expected type '%d'."),
 				*OptionalAttributes[Idx].Name,
 				(int)FoundAttribute->GetType(),
 				(int)OptionalAttributes[Idx].Type);
@@ -164,7 +164,7 @@ bool FAttributeValidator::Validate(const NNE::FAttributeMap& AttributesToValidat
 		if (OptionalAttribute == nullptr && RequiredAttribute == nullptr)
 		{
 			bAreAttributesValid = false;
-			UE_LOG(LogNNE, Warning, TEXT("Found unsupported attribute '%s'."), *Name);
+			UE_LOG(LogNNERuntimeRDGHlsl, Warning, TEXT("Found unsupported attribute '%s'."), *Name);
 		}
 	}
 
diff --git a/Engine/Plugins/Experimental/NNERuntimeRDG/Source/NNERuntimeRDG/Private/NNERuntimeRDGBase.h b/Engine/Plugins/Experimental/NNERuntimeRDG/Source/NNERuntimeRDG/Private/NNERuntimeRDGBase.h
--- a/Engine/Plugins/Experimental/NNERuntimeRDG/Source/NNERuntimeRDG/Private/NNERuntimeRDGBase.h
+++ b/Engine/Plugins/Experimental/NNERuntimeRDG/Source/NNERuntimeRDG/Private/NNERuntimeRDGBase.h
@@ -4,11 +4,11 @@
 
 #include "NNEAttributeMap.h"
 #include "NNEAttributeValue.h"
-#include "NNETensor.h"
-#include "NNETypes.h"
-#include "NNE.h"
+#include "NNEHlslShadersLog.h"
 #include "NNEModelOptimizerInterface.h"
 #include "NNERuntimeFormat.h"
+#include "NNETensor.h"
+#include "NNETypes.h"
 #include "RenderGraphResources.h"
 #include "Serialization/MemoryReader.h"
 #include "ShaderParameterUtils.h"
@@ -187,7 +187,7 @@ public:
 		const FOperatorFunctions* OperatorFunctions = OpFindFunctions(OpDesc);
 		if(OperatorFunctions == nullptr)
 		{
-			UE_LOG(LogNNE, Warning, TEXT("RDG MLOperator: %s is not registered"), *OpDesc.GetFullName());
+			UE_LOG(LogNNERuntimeRDGHlsl, Warning, TEXT("Operator: %s is not registered"), *OpDesc.GetFullName());
 			return nullptr;
 		}
 
@@ -199,7 +199,7 @@ public:
 		const FOperatorFunctions* OperatorFunctions = OpFindFunctions(OpDesc);
 		if(OperatorFunctions == nullptr)
 		{
-			UE_LOG(LogNNE, Warning, TEXT("RDG MLOperator: %s is not registered"), *OpDesc.GetFullName());
+			UE_LOG(LogNNERuntimeRDGHlsl, Warning, TEXT("Operator: %s is not registered"), *OpDesc.GetFullName());
 			return nullptr;
 		}
 
@@ -214,12 +214,12 @@ public:
 			const FOperatorFunctions* OperatorFunctions = VersionToFunctions->Find(OpDesc.Version);
 			if(OperatorFunctions != nullptr)
 			{
-				UE_LOG(LogNNE, Warning, TEXT("RDG MLOperator is already registered: %s"), *OpDesc.GetFullName());
+				UE_LOG(LogNNERuntimeRDGHlsl, Warning, TEXT("Operator is already registered: %s"), *OpDesc.GetFullName());
 				return false;
 			}
 			else if(!OpDesc.Version.IsSet() || VersionToFunctions->Find(TOptional<TOperatorVersionType>()) != nullptr)
 			{
-				UE_LOG(LogNNE, Warning, TEXT("RDG MLOperator %s is unversioned, can't register a version of it"), *FOperatorDesc{{OpDesc}}.GetFullName());
+				UE_LOG(LogNNERuntimeRDGHlsl, Warning, TEXT("Operator %s is unversioned, can't register a version of it"), *FOperatorDesc{{OpDesc}}.GetFullName());
 				return false;
 			}
 			VersionToFunctions->Add(OpDesc.Version, FOperatorFunctions{CreateFunc, ValidateFunc});
@@ -329,7 +329,7 @@ private:
 			}
 		}
 
-		UE_LOG(LogNNE, Warning, TEXT("RDG MLOperator: %s is not registered"), *OpDesc.GetFullName());
+		UE_LOG(LogNNERuntimeRDGHlsl, Warning, TEXT("Operator: %s is not registered"), *OpDesc.GetFullName());
 		return nullptr;
 	}
 
@@ -358,7 +358,7 @@ public:
 		ENNEInferenceFormat FormatType = InputModel.Format;
 		if (FormatType != ENNEInferenceFormat::NNERT)
 		{
-			UE_LOG(LogNNE, Warning, TEXT("Unsupported format type for validator %s"), *GetName());
+			UE_LOG(LogNNERuntimeRDGHlsl, Warning, TEXT("Unsupported format type for validator %s"), *GetName());
 			return false;
 		}
 
@@ -392,13 +392,13 @@ public:
 
 			if (!ValidationFn)
 			{
-				UE_LOG(LogNNE, Warning, TEXT("RDG MLOperatorRegistry failed to find validation for operator:%s"), *OpType);
+				UE_LOG(LogNNERuntimeRDGHlsl, Warning, TEXT("OperatorRegistry failed to find validation for operator:%s"), *OpType);
 				return false;
 			}
 			
 			if (!ValidationFn(AttributeMap, InputTensorTypes, InputTensorShapes))
 			{
-				UE_LOG(LogNNE, Warning, TEXT("RDG MLOperatorRegistry failed to validate operator:%s"), *OpType);
+				UE_LOG(LogNNERuntimeRDGHlsl, Warning, TEXT("OperatorRegistry failed to validate operator:%s"), *OpType);
 				return false;
 			}
 		}
diff --git a/Engine/Plugins/Experimental/NNERuntimeRDG/Source/NNERuntimeRDG/Private/NNERuntimeRDGHlsl.cpp b/Engine/Plugins/Experimental/NNERuntimeRDG/Source/NNERuntimeRDG/Private/NNERuntimeRDGHlsl.cpp
--- a/Engine/Plugins/Experimental/NNERuntimeRDG/Source/NNERuntimeRDG/Private/NNERuntimeRDGHlsl.cpp
+++ b/Engine/Plugins/Experimental/NNERuntimeRDG/Source/NNERuntimeRDG/Private/NNERuntimeRDGHlsl.cpp
@@ -3,8 +3,8 @@
 #include "NNERuntimeRDGHlsl.h"
 
 #include "Misc/SecureHash.h"
-#include "NNE.h"
 #include "NNEAttributeMap.h"
+#include "NNEHlslShadersLog.h"
 #include "NNEModelData.h"
 #include "NNEModelOptimizerInterface.h"
 #include "NNERuntimeRDG.h"
@@ -111,7 +111,7 @@ namespace ConsoleCommands
 				FOperatorRegistryHlsl* Registry = FOperatorRegistryHlsl::Get();
 				check(Registry != nullptr);
 				FString SupportMatrix = Registry->ListAllRegisteredOperators();
-				UE_LOG(LogNNE, Display, TEXT("HLSL Operators support matrix: \n%s"), *SupportMatrix);
+				UE_LOG(LogNNERuntimeRDGHlsl, Display, TEXT("Operators support matrix: \n%s"), *SupportMatrix);
 			}
 		)
 	);
@@ -124,19 +124,19 @@ bool UNNERuntimeRDGHlslImpl::IsCurrentPlatformSupported()
 #ifndef NNE_FORCE_HARDWARE_SUPPORTS_HLSL
 	if(GMaxRHIFeatureLevel < ERHIFeatureLevel::SM5)
 	{
-		UE_LOG(LogNNE, Display, TEXT("UNNERuntimeRDGHlsl: minimum feature level required is SM5 for current RHI platform."));
+		UE_LOG(LogNNERuntimeRDGHlsl, Display, TEXT("Minimum feature level required is SM5 for current RHI platform."));
 		bResult = false;
 	}
 
 	if(!GRHISupportsWaveOperations)
 	{
-		UE_LOG(LogNNE, Display, TEXT("UNNERuntimeRDGHlsl: current RHI platform doesn't support wave operations."));
+		UE_LOG(LogNNERuntimeRDGHlsl, Display, TEXT("Current RHI platform doesn't support wave operations."));
 		bResult = false;
 	}
 
 	if(!GRHIGlobals.SupportsNative16BitOps)
 	{
-		UE_LOG(LogNNE, Display, TEXT("UNNERuntimeRDGHlsl: current RHI platform doesn't support native 16-bit operations."));
+		UE_LOG(LogNNERuntimeRDGHlsl, Display, TEXT("Current RHI platform doesn't support native 16-bit operations."));
 		bResult = false;
 	}
 #endif
@@ -149,7 +149,7 @@ UNNERuntimeRDGHlslImpl::ECanCreateModelDataStatus UNNERuntimeRDGHlslImpl::CanCre
 #ifdef NNE_UTILITIES_AVAILABLE
 	return FileType.Compare("onnx", ESearchCase::IgnoreCase) == 0 ? ECanCreateModelDataStatus::Ok : ECanCreateModelDataStatus::FailFileIdNotSupported;
 #else
-	UE_LOG(LogNNE, Display, TEXT("NNERuntimeRDGUtils is not available on this platform"));
+	UE_LOG(LogNNERuntimeRDGHlsl, Display, TEXT("NNERuntimeRDGUtils is not available on this platform"));
 	return ECanCreateModelDataStatus::Fail;
 #endif
 }
@@ -184,19 +184,19 @@ TSharedPtr<UE::NNE::FSharedModelData> UNNERuntimeRDGHlslImpl::CreateModelData(co
 	// Check model is not > 2GB
 	if ((TArray<uint8>::SizeType)FileData.Num() != FileData.Num())
 	{
-		UE_LOG(LogNNE, Warning, TEXT("UNNERuntimeRDGHlsl cannot create the model data with id %s (Filetype: %s), models > 2GBs are not supported"), *FileId.ToString(EGuidFormats::Digits).ToLower(), *FileType);
+		UE_LOG(LogNNERuntimeRDGHlsl, Warning, TEXT("Cannot create the model data with id %s (Filetype: %s), models > 2GBs are not supported"), *FileId.ToString(EGuidFormats::Digits).ToLower(), *FileType);
 		return {};
 	}
 
 	if (!AdditionalFileData.IsEmpty())
 	{
-		UE_LOG(LogNNE, Warning, TEXT("UNNERuntimeRDGHlsl cannot create the model data with id %s (Filetype: %s), external data not supported at the moment, please convert the model to internal storage. See https://onnx.ai/onnx/repo-docs/ExternalData.html"), *FileId.ToString(EGuidFormats::Digits).ToLower(), *FileType);
+		UE_LOG(LogNNERuntimeRDGHlsl, Warning, TEXT("Cannot create the model data with id %s (Filetype: %s), external data not supported at the moment, please convert the model to internal storage. See https://onnx.ai/onnx/repo-docs/ExternalData.html"), *FileId.ToString(EGuidFormats::Digits).ToLower(), *FileType);
 		return {};
 	}
 
 	if (CanCreateModelData(FileType, FileData, AdditionalFileData, FileId, TargetPlatform) != ECanCreateModelDataStatus::Ok)
 	{
-		UE_LOG(LogNNE, Warning, TEXT("UNNERuntimeRDGHlsl cannot create the model data with id %s (Filetype: %s)"), *FileId.ToString(EGuidFormats::Digits).ToLower(), *FileType);
+		UE_LOG(LogNNERuntimeRDGHlsl, Warning, TEXT("Cannot create the model data with id %s (Filetype: %s)"), *FileId.ToString(EGuidFormats::Digits).ToLower(), *FileType);
 		return {};
 	}
 
@@ -236,7 +236,7 @@ TSharedPtr<UE::NNE::IModelRDG> UNNERuntimeRDGHlslImpl::CreateModelRDG(const TObj
 {
 	if (CanCreateModelRDG(ModelData) != ECanCreateModelRDGStatus::Ok)
 	{
-		UE_LOG(LogNNE, Warning, TEXT("UNNERuntimeRDGHlsl cannot create a model from the model data with id %s"), *ModelData->GetFileId().ToString(EGuidFormats::Digits));
+		UE_LOG(LogNNERuntimeRDGHlsl, Warning, TEXT("Cannot create a model from the model data with id %s"), *ModelData->GetFileId().ToString(EGuidFormats::Digits));
 		return TSharedPtr<UE::NNE::IModelRDG>();
 	}
 
diff --git a/Engine/Plugins/Experimental/NNERuntimeRDG/Source/NNERuntimeRDG/Private/NNERuntimeRDGModel.cpp b/Engine/Plugins/Experimental/NNERuntimeRDG/Source/NNERuntimeRDG/Private/NNERuntimeRDGModel.cpp
--- a/Engine/Plugins/Experimental/NNERuntimeRDG/Source/NNERuntimeRDG/Private/NNERuntimeRDGModel.cpp
+++ b/Engine/Plugins/Experimental/NNERuntimeRDG/Source/NNERuntimeRDG/Private/NNERuntimeRDGModel.cpp
@@ -2,6 +2,7 @@
 
 #include "NNERuntimeRDGModel.h"
 
+#include "NNEHlslShadersLog.h"
 #include "NNERuntimeFormat.h"
 #include "RenderGraphBuilder.h"
 #include "RenderGraphUtils.h"
@@ -47,7 +48,7 @@ bool FModelInstanceRDG::LoadModel(TConstArrayView<uint8> ModelData, FNNERuntimeF
 
 		if (FormatTensorDesc.Type != ENNEFormatTensorType::Empty && FormatTensorDesc.DataType == ENNETensorDataType::None)
 		{
-			UE_LOG(LogNNE, Error, TEXT("Tensor %s has invalid format: Data type None is reserved for empty tensors."), *SymbolicTensor.GetName());
+			UE_LOG(LogNNERuntimeRDGHlsl, Error, TEXT("Tensor %s has invalid format: Data type None is reserved for empty tensors."), *SymbolicTensor.GetName());
 			return false;
 		}
 
@@ -76,7 +77,7 @@ bool FModelInstanceRDG::LoadModel(TConstArrayView<uint8> ModelData, FNNERuntimeF
 			WeightTensorIndices.Emplace(Idx);
 			if (!SymbolicTensor.GetShape().IsConcrete())
 			{
-				UE_LOG(LogNNE, Error, TEXT("Weight tensor %s should have a concrete shape"), *SymbolicTensor.GetName());
+				UE_LOG(LogNNERuntimeRDGHlsl, Error, TEXT("Weight tensor %s should have a concrete shape"), *SymbolicTensor.GetName());
 				return false;
 			}
 
@@ -85,7 +86,7 @@ bool FModelInstanceRDG::LoadModel(TConstArrayView<uint8> ModelData, FNNERuntimeF
 
 			if (WeightRDG.GetDataSize() != FormatTensorDesc.DataSize)
 			{
-				UE_LOG(LogNNE, Error, TEXT("Weight %s has incorrect size. Expected %d bytes, got %d"), *SymbolicTensor.GetName(), FormatTensorDesc.DataSize, WeightRDG.GetDataSize());
+				UE_LOG(LogNNERuntimeRDGHlsl, Error, TEXT("Weight %s has incorrect size. Expected %d bytes, got %d"), *SymbolicTensor.GetName(), FormatTensorDesc.DataSize, WeightRDG.GetDataSize());
 				return false;
 			}
 
@@ -229,21 +230,21 @@ FModelInstanceRDG::EEnqueueRDGStatus FModelInstanceRDG::EnqueueRDG(FRDGBuilder&
 	// Verify the model inputs were prepared
 	if (InputTensorShapes.Num() == 0)
 	{
-		UE_LOG(LogNNE, Error, TEXT("EnqueueRDG(): Input shapes are not set, please call SetInputTensorShapes."));
+		UE_LOG(LogNNERuntimeRDGHlsl, Error, TEXT("Input shapes are not set, please call SetInputTensorShapes."));
 		return EEnqueueRDGStatus::Fail;
 	}
 
 	Res = SetTensors(RDGBuilder, InputTensorRDGs, InInputBindings);
 	if (Res != -1)
 	{
-		UE_LOG(LogNNE, Warning, TEXT("Invalid buffer for input tensor binding at index %d"), Res);
+		UE_LOG(LogNNERuntimeRDGHlsl, Warning, TEXT("Invalid buffer for input tensor binding at index %d"), Res);
 		return EEnqueueRDGStatus::Fail;
 	}
 
 	Res = SetTensors(RDGBuilder, OutputTensorRDGs, InOutputBindings);
 	if (Res != -1)
 	{
-		UE_LOG(LogNNE, Warning, TEXT("Invalid buffer for output tensor binding at index %d"), Res);
+		UE_LOG(LogNNERuntimeRDGHlsl, Warning, TEXT("Invalid buffer for output tensor binding at index %d"), Res);
 		return EEnqueueRDGStatus::Fail;
 	}
 
@@ -291,12 +292,12 @@ int32 FModelInstanceRDG::SetTensors(FRDGBuilder& GraphBuilder, FTensorRDGArray&
 		const NNE::FTensorBindingRDG& Binding = InBindings[Idx];
 		if (Binding.Buffer == nullptr)
 		{
-			UE_LOG(LogNNE, Warning, TEXT("nullptr buffer encountered."));
+			UE_LOG(LogNNERuntimeRDGHlsl, Warning, TEXT("nullptr buffer encountered."));
 			return Idx;
 		}
 		if (Binding.Buffer->GetSize() % NNERUNTIMERDGHLSL_BUFFER_LENGTH_ALIGNMENT != 0)
 		{
-			UE_LOG(LogNNE, Warning, TEXT("Buffer has size %i which is not a multiple of %i"), Binding.Buffer->GetSize(), NNERUNTIMERDGHLSL_BUFFER_LENGTH_ALIGNMENT);
+			UE_LOG(LogNNERuntimeRDGHlsl, Warning, TEXT("Buffer has size %i which is not a multiple of %i"), Binding.Buffer->GetSize(), NNERUNTIMERDGHLSL_BUFFER_LENGTH_ALIGNMENT);
 			return Idx;
 		}
 		TensorRDG.SetBuffer(Binding.Buffer);
diff --git a/Engine/Plugins/Experimental/NNERuntimeRDG/Source/NNERuntimeRDG/Private/NNERuntimeRDGModelHlsl.cpp b/Engine/Plugins/Experimental/NNERuntimeRDG/Source/NNERuntimeRDG/Private/NNERuntimeRDGModelHlsl.cpp
--- a/Engine/Plugins/Experimental/NNERuntimeRDG/Source/NNERuntimeRDG/Private/NNERuntimeRDGModelHlsl.cpp
+++ b/Engine/Plugins/Experimental/NNERuntimeRDG/Source/NNERuntimeRDG/Private/NNERuntimeRDGModelHlsl.cpp
@@ -1,6 +1,8 @@
 // Copyright Epic Games, Inc. All Rights Reserved.
 
 #include "NNERuntimeRDGModelHlsl.h"
+
+#include "NNEHlslShadersLog.h"
 #include "NNETensor.h"
 #include "NNERuntimeRDGHlsl.h"
 #include "NNERuntimeRDGHlslOp.h"
@@ -19,7 +21,7 @@ FOperatorHlsl* OpCreate(const FOperatorDesc& OpDesc, TConstArrayView<NNE::FTenso
 
 	if (!CreateFn)
 	{
-		UE_LOG(LogNNE, Warning, TEXT("Hlsl MLOperatorRegistry failed to find operator: %s"), *OpDesc.GetFullName());
+		UE_LOG(LogNNERuntimeRDGHlsl, Warning, TEXT("Operator registry failed to find operator: %s"), *OpDesc.GetFullName());
 		return nullptr;
 	}
 
@@ -27,7 +29,7 @@ FOperatorHlsl* OpCreate(const FOperatorDesc& OpDesc, TConstArrayView<NNE::FTenso
 
 	if (!Op->Initialize(InputTensorDescs, OutputTensorDescs, AttributeMap))
 	{
-		UE_LOG(LogNNE, Warning, TEXT("Hlsl runtime: Error initializing operator: %s"), *OpDesc.GetFullName());
+		UE_LOG(LogNNERuntimeRDGHlsl, Warning, TEXT("Error initializing operator: %s"), *OpDesc.GetFullName());
 		delete Op;
 		return nullptr;
 	}
@@ -166,7 +168,7 @@ bool FModelInstance::Init(TConstArrayView<uint8> ModelData)
 
 		if (!Op) //Op.Shader.IsNull())
 		{
-			UE_LOG(LogNNE, Warning, TEXT("Failed to create operator:%s"), *OperatorDesc.GetFullName());
+			UE_LOG(LogNNERuntimeRDGHlsl, Warning, TEXT("Failed to create operator:%s"), *OperatorDesc.GetFullName());
 
 			//Note: Need to cleanup operators
 			return false;
@@ -230,7 +232,7 @@ int FModelInstance::PrepareTensorShapesAndData()
 	
 	if (Operators.Num() == 0)
 	{
-		UE_LOG(LogNNE, Warning, TEXT("No operators in model"));
+		UE_LOG(LogNNERuntimeRDGHlsl, Warning, TEXT("No operators in model"));
 		return -1;
 	}
 
@@ -288,7 +290,7 @@ int FModelInstance::PrepareTensorShapesAndData()
 		{
 			//Operator could not prepare the output tensors, meaning we can't allocate
 			//output buffer before running the model. This runtime does not support this.
-			UE_LOG(LogNNE, Warning, TEXT("Could not deduce tensor shapes for this model during shape inference, HLSL runtime wont support the model as it need to precompute all shapes for performance reasons."));
+			UE_LOG(LogNNERuntimeRDGHlsl, Warning, TEXT("Could not deduce tensor shapes for this model during shape inference, HLSL runtime wont support the model as it need to precompute all shapes for performance reasons."));
 			AllTensorRDGRefs.Reset();
 			return -1;
 		}
diff --git a/Engine/Plugins/Experimental/NNERuntimeRDG/Source/NNERuntimeRDG/Private/NNERuntimeRDGModule.cpp b/Engine/Plugins/Experimental/NNERuntimeRDG/Source/NNERuntimeRDG/Private/NNERuntimeRDGModule.cpp
--- a/Engine/Plugins/Experimental/NNERuntimeRDG/Source/NNERuntimeRDG/Private/NNERuntimeRDGModule.cpp
+++ b/Engine/Plugins/Experimental/NNERuntimeRDG/Source/NNERuntimeRDG/Private/NNERuntimeRDGModule.cpp
@@ -1,8 +1,10 @@
 // Copyright Epic Games, Inc. All Rights Reserved.
 
-#include "HAL/IConsoleManager.h"
 #include "NNERuntimeRDGModule.h"
+
+#include "HAL/IConsoleManager.h"
 #include "NNE.h"
+#include "NNEHlslShadersLog.h"
 #include "NNERuntimeRDGHlsl.h"
 #include "UObject/WeakInterfacePtr.h"
 #include "DataDrivenShaderPlatformInfo.h"
@@ -47,15 +49,15 @@ void FNNERuntimeRDGModule::StartupModule()
 			}
 			else
 			{
-				UE_LOG(LogNNE, Display, TEXT("FNNERuntimeRDGModule: not registering NNERuntimeRDGHlsl runtime because current hardware is incompatible, consider bypassing by setting the define NNE_FORCE_HARDWARE_SUPPORTS_HLSL."));
+				UE_LOG(LogNNERuntimeRDGHlsl, Display, TEXT("Not registering runtime because current hardware is incompatible, consider bypassing by setting the define NNE_FORCE_HARDWARE_SUPPORTS_HLSL."));
 			}
 		}
 		else
 		{
-			UE_LOG(LogNNE, Display, TEXT("FNNERuntimeRDGModule: not registering NNERuntimeRDGHlsl runtime because current RHI shader platform is not enabled, consider setting the flag bSupportsNNEShaders in DataDrivenPlatformInfo."));
+			UE_LOG(LogNNERuntimeRDGHlsl, Display, TEXT("Not registering runtime because current RHI shader platform is not enabled, consider setting the flag bSupportsNNEShaders in DataDrivenPlatformInfo."));
 		}
 	#else
-		UE_LOG(LogNNE, Display, TEXT("FNNERuntimeRDGModule: not registering NNERuntimeRDGHlsl runtime as platform is not enabled, if needed set define WITH_NNE_RUNTIME_HLSL."));
+		UE_LOG(LogNNERuntimeRDGHlsl, Display, TEXT("Not registering runtime as platform is not enabled, if needed set define WITH_NNE_RUNTIME_HLSL."));
 	#endif
 #endif
 }
diff --git a/Engine/Plugins/Experimental/NNERuntimeRDG/Source/NNERuntimeRDGUtils/NNERuntimeRDGUtils.Build.cs b/Engine/Plugins/Experimental/NNERuntimeRDG/Source/NNERuntimeRDGUtils/NNERuntimeRDGUtils.Build.cs
--- a/Engine/Plugins/Experimental/NNERuntimeRDG/Source/NNERuntimeRDGUtils/NNERuntimeRDGUtils.Build.cs
+++ b/Engine/Plugins/Experimental/NNERuntimeRDG/Source/NNERuntimeRDGUtils/NNERuntimeRDGUtils.Build.cs
@@ -19,6 +19,7 @@ public class NNERuntimeRDGUtils : ModuleRules
 				"CoreUObject",
 				"Engine",
 				"Projects",
+				"NNEHlslShaders",
 				"NNERuntimeRDGProtobufEditor",
 				"NNERuntimeRDGOnnxruntimeEditor",
 				"NNERuntimeRDGOnnxEditor"
diff --git a/Engine/Plugins/Experimental/NNERuntimeRDG/Source/NNERuntimeRDGUtils/Private/NNERuntimeRDGUtilsHelpers.cpp b/Engine/Plugins/Experimental/NNERuntimeRDG/Source/NNERuntimeRDGUtils/Private/NNERuntimeRDGUtilsHelpers.cpp
--- a/Engine/Plugins/Experimental/NNERuntimeRDG/Source/NNERuntimeRDGUtils/Private/NNERuntimeRDGUtilsHelpers.cpp
+++ b/Engine/Plugins/Experimental/NNERuntimeRDG/Source/NNERuntimeRDGUtils/Private/NNERuntimeRDGUtilsHelpers.cpp
@@ -2,7 +2,7 @@
 
 #include "NNERuntimeRDGUtilsHelpers.h"
 
-#include "NNE.h"
+#include "NNEHlslShadersLog.h"
 
 THIRD_PARTY_INCLUDES_START
 #include <onnx/defs/schema.h>
@@ -17,7 +17,7 @@ TOptional<uint32> GetOpVersionFromOpsetVersion(const FString& OpType, int OpsetV
 	const onnx::OpSchema* OpSchema = onnx::OpSchemaRegistry::Schema(TCHAR_TO_ANSI(*OpType), OpsetVersion);
 	if(OpSchema == nullptr)
 	{
-		UE_LOG(LogNNE, Warning, TEXT("No OpSchema found for operator %s and OpSet version %d."), *OpType, OpsetVersion);
+		UE_LOG(LogNNERuntimeRDGHlsl, Warning, TEXT("No OpSchema found for operator %s and OpSet version %d."), *OpType, OpsetVersion);
 		return TOptional<uint32>();
 	}
 	return (uint32) OpSchema->SinceVersion();
diff --git a/Engine/Plugins/Experimental/NNERuntimeRDG/Source/NNERuntimeRDGUtils/Private/NNERuntimeRDGUtilsModelBuilderNNE.cpp b/Engine/Plugins/Experimental/NNERuntimeRDG/Source/NNERuntimeRDGUtils/Private/NNERuntimeRDGUtilsModelBuilderNNE.cpp
--- a/Engine/Plugins/Experimental/NNERuntimeRDG/Source/NNERuntimeRDGUtils/Private/NNERuntimeRDGUtilsModelBuilderNNE.cpp
+++ b/Engine/Plugins/Experimental/NNERuntimeRDG/Source/NNERuntimeRDGUtils/Private/NNERuntimeRDGUtilsModelBuilderNNE.cpp
@@ -2,8 +2,8 @@
 
 #include "NNERuntimeRDGUtilsModelBuilderNNE.h"
 
-#include "NNE.h"
 #include "NNEAttributeMap.h"
+#include "NNEHlslShadersLog.h"
 #include "NNERuntimeFormat.h"
 #include "Misc/StringBuilder.h"
 #include "Serialization/MemoryWriter.h"
@@ -88,13 +88,13 @@ public:
 
 		if (Idx < 0 || Idx >= Format.Tensors.Num())
 		{
-			UE_LOG(LogNNE, Warning, TEXT("Failed to add input tensor, invalid tensor index"));
+			UE_LOG(LogNNERuntimeRDGHlsl, Warning, TEXT("Failed to add input tensor, invalid tensor index"));
 			return false;
 		}
 
 		if (Format.Tensors[Idx].Type != ENNEFormatTensorType::None)
 		{
-			UE_LOG(LogNNE, Warning, TEXT("Failed to add input tensor, tensor usage already set up"));
+			UE_LOG(LogNNERuntimeRDGHlsl, Warning, TEXT("Failed to add input tensor, tensor usage already set up"));
 			return false;
 		}
 
@@ -109,13 +109,13 @@ public:
 
 		if (Idx < 0 || Idx >= Format.Tensors.Num())
 		{
-			UE_LOG(LogNNE, Warning, TEXT("Failed to add output tensor, invalid tensor index"));
+			UE_LOG(LogNNERuntimeRDGHlsl, Warning, TEXT("Failed to add output tensor, invalid tensor index"));
 			return false;
 		}
 
 		if (Format.Tensors[Idx].Type != ENNEFormatTensorType::None)
 		{
-			UE_LOG(LogNNE, Warning, TEXT("Failed to add output tensor, tensor usage already set up"));
+			UE_LOG(LogNNERuntimeRDGHlsl, Warning, TEXT("Failed to add output tensor, tensor usage already set up"));
 			return false;
 		}
 
@@ -145,7 +145,7 @@ public:
 
 		if (TensorIdx < 0 || TensorIdx >= Format.Tensors.Num())
 		{
-			UE_LOG(LogNNE, Warning, TEXT("Failed to add operator input tensor, invalid tensor index"));
+			UE_LOG(LogNNERuntimeRDGHlsl, Warning, TEXT("Failed to add operator input tensor, invalid tensor index"));
 			return false;
 		}
 
@@ -161,13 +161,13 @@ public:
 
 		if (TensorIdx < 0 || TensorIdx >= Format.Tensors.Num())
 		{
-			UE_LOG(LogNNE, Warning, TEXT("Failed to add operator output tensor, invalid tensor index"));
+			UE_LOG(LogNNERuntimeRDGHlsl, Warning, TEXT("Failed to add operator output tensor, invalid tensor index"));
 			return false;
 		}
 
 		if (Format.Tensors[TensorIdx].Type == ENNEFormatTensorType::Input)
 		{
-			UE_LOG(LogNNE, Warning, TEXT("Failed to add output tensor, tensor usage already set up to input"));
+			UE_LOG(LogNNERuntimeRDGHlsl, Warning, TEXT("Failed to add output tensor, tensor usage already set up to input"));
 			return false;
 		}
 
diff --git a/Engine/Plugins/Experimental/NNERuntimeRDG/Source/NNERuntimeRDGUtils/Private/NNERuntimeRDGUtilsModelBuilderONNX.cpp b/Engine/Plugins/Experimental/NNERuntimeRDG/Source/NNERuntimeRDGUtils/Private/NNERuntimeRDGUtilsModelBuilderONNX.cpp
--- a/Engine/Plugins/Experimental/NNERuntimeRDG/Source/NNERuntimeRDGUtils/Private/NNERuntimeRDGUtilsModelBuilderONNX.cpp
+++ b/Engine/Plugins/Experimental/NNERuntimeRDG/Source/NNERuntimeRDGUtils/Private/NNERuntimeRDGUtilsModelBuilderONNX.cpp
@@ -2,8 +2,8 @@
 
 #include "NNERuntimeRDGUtilsModelBuilderONNX.h"
 
-#include "NNE.h"
 #include "NNEAttributeMap.h"
+#include "NNEHlslShadersLog.h"
 #include "NNERuntimeFormat.h"
 
 THIRD_PARTY_INCLUDES_START
@@ -103,7 +103,7 @@ public:
 			}
 			catch (onnx::checker::ValidationError& e)
 			{
-				UE_LOG(LogNNE, Warning, TEXT("ModelBuilder error:%s"), ANSI_TO_TCHAR(e.what()));
+				UE_LOG(LogNNERuntimeRDGHlsl, Warning, TEXT("ModelBuilder error:%s"), ANSI_TO_TCHAR(e.what()));
 				return false;
 			}
 #endif
diff --git a/Engine/Plugins/Experimental/NNERuntimeRDG/Source/NNERuntimeRDGUtils/Private/NNERuntimeRDGUtilsModelOptimizerBase.cpp b/Engine/Plugins/Experimental/NNERuntimeRDG/Source/NNERuntimeRDGUtils/Private/NNERuntimeRDGUtilsModelOptimizerBase.cpp
--- a/Engine/Plugins/Experimental/NNERuntimeRDG/Source/NNERuntimeRDGUtils/Private/NNERuntimeRDGUtilsModelOptimizerBase.cpp
+++ b/Engine/Plugins/Experimental/NNERuntimeRDG/Source/NNERuntimeRDGUtils/Private/NNERuntimeRDGUtilsModelOptimizerBase.cpp
@@ -2,7 +2,7 @@
 
 #include "NNERuntimeRDGUtilsModelOptimizerBase.h"
 
-#include "NNE.h"
+#include "NNEHlslShadersLog.h"
 #include "NNERuntimeFormat.h"
 
 THIRD_PARTY_INCLUDES_START
@@ -38,7 +38,7 @@ bool FModelValidatorONNX::ValidateModel(const FNNEModelRaw& InputModel, const FO
 	}
 	catch (onnx::checker::ValidationError& e)
 	{
-		UE_LOG(LogNNE, Warning, TEXT("Input model is invalid : %s."), ANSI_TO_TCHAR(e.what()));
+		UE_LOG(LogNNERuntimeRDGHlsl, Warning, TEXT("Input model is invalid : %s."), ANSI_TO_TCHAR(e.what()));
 		return false;
 	}
 #endif
@@ -55,7 +55,7 @@ bool FModelOptimizerBase::IsModelValid(const FNNEModelRaw& ModelToValidate, cons
 		check(Validator.IsValid());
 		if (!Validator->ValidateModel(ModelToValidate, Options))
 		{
-			UE_LOG(LogNNE, Warning, TEXT("Model validator '%s' detected an error."), *(Validator->GetName()));
+			UE_LOG(LogNNERuntimeRDGHlsl, Warning, TEXT("Model validator '%s' detected an error."), *(Validator->GetName()));
 			bIsModelValid = false;
 		}
 	}
@@ -66,7 +66,7 @@ bool FModelOptimizerBase::ApplyAllPassesAndValidations(FNNEModelRaw& OptimizedMo
 {
 	if (!IsModelValid(OptimizedModel, Options))
 	{
-		UE_LOG(LogNNE, Warning, TEXT("Model is not valid."));
+		UE_LOG(LogNNERuntimeRDGHlsl, Warning, TEXT("Model is not valid."));
 		return false;
 	}
 		
@@ -79,7 +79,7 @@ bool FModelOptimizerBase::ApplyAllPassesAndValidations(FNNEModelRaw& OptimizedMo
 			
 		if (!Pass->ApplyPass(OptimizedModel, Options))
 		{
-			UE_LOG(LogNNE, Warning, TEXT("Error while executing model optimisation pass '%s'."), *(Pass->GetName()));
+			UE_LOG(LogNNERuntimeRDGHlsl, Warning, TEXT("Error while executing model optimisation pass '%s'."), *(Pass->GetName()));
 			return false;
 		}
 
@@ -88,7 +88,7 @@ bool FModelOptimizerBase::ApplyAllPassesAndValidations(FNNEModelRaw& OptimizedMo
 
 		if (!IsModelValid(OptimizedModel, Options))
 		{
-			UE_LOG(LogNNE, Warning, TEXT("Model validation failed after optimisation pass '%s'."), *(Pass->GetName()));
+			UE_LOG(LogNNERuntimeRDGHlsl, Warning, TEXT("Model validation failed after optimisation pass '%s'."), *(Pass->GetName()));
 			return false;
 		}
 	}
diff --git a/Engine/Plugins/Experimental/NNERuntimeRDG/Source/NNERuntimeRDGUtils/Private/NNERuntimeRDGUtilsModelOptimizerNNE.cpp b/Engine/Plugins/Experimental/NNERuntimeRDG/Source/NNERuntimeRDGUtils/Private/NNERuntimeRDGUtilsModelOptimizerNNE.cpp
--- a/Engine/Plugins/Experimental/NNERuntimeRDG/Source/NNERuntimeRDGUtils/Private/NNERuntimeRDGUtilsModelOptimizerNNE.cpp
+++ b/Engine/Plugins/Experimental/NNERuntimeRDG/Source/NNERuntimeRDGUtils/Private/NNERuntimeRDGUtilsModelOptimizerNNE.cpp
@@ -3,13 +3,13 @@
 #include "NNERuntimeRDGUtilsModelOptimizerNNE.h"
 
 #include "Containers/ContainersFwd.h"
-#include "NNE.h"
+#include "NNEAttributeTensor.h"
+#include "NNEHlslShadersLog.h"
 #include "NNERuntimeFormat.h"
 #include "NNERuntimeRDGUtilsModelBuilderNNE.h"
 #include "NNERuntimeRDGUtilsModelOptimizerONNX.h"
 #include "NNERuntimeRDGUtilsHelpers.h"
 #include "NNETensor.h"
-#include "NNEAttributeTensor.h"
 #include "NNETypes.h"
 
 THIRD_PARTY_INCLUDES_START
@@ -192,7 +192,7 @@ namespace ModelOptimizerNNEHelper
 		}
 		else if (Tensor.external_data().size()) 
 		{
-			UE_LOG(LogNNE, Warning, TEXT("GetTensorInfoFromONNXInitializer: External data not supported."));
+			UE_LOG(LogNNERuntimeRDGHlsl, Warning, TEXT("External data not supported."));
 			return false;
 		}
 		else if (Tensor.float_data().size())
@@ -262,12 +262,12 @@ namespace ModelOptimizerNNEHelper
 
 		if (!GetTensorInfoFromONNXInitializer(InTensor, Shape, DataType, Data))
 		{
-			UE_LOG(LogNNE, Error, TEXT("Tensor data could not be loaded"));
+			UE_LOG(LogNNERuntimeRDGHlsl, Error, TEXT("Tensor data could not be loaded"));
 			return false;
 		}
 		if (Shape.Num() > FTensorShape::MaxRank)
 		{
-			UE_LOG(LogNNE, Error, TEXT("Tensor shape of rank %i exceeds MaxRank %i"), Shape.Num(), FTensorShape::MaxRank);
+			UE_LOG(LogNNERuntimeRDGHlsl, Error, TEXT("Tensor shape of rank %i exceeds MaxRank %i"), Shape.Num(), FTensorShape::MaxRank);
 			return false;
 		}
 		TArray<uint32, TInlineAllocator<FTensorShape::MaxRank>> UIntShape;
@@ -275,7 +275,7 @@ namespace ModelOptimizerNNEHelper
 		{
 			if (value < 0)
 			{
-				UE_LOG(LogNNE, Error, TEXT("Tensor shape has negative value"));
+				UE_LOG(LogNNERuntimeRDGHlsl, Error, TEXT("Tensor shape has negative value"));
 				return false;
 			}
 			UIntShape.Add(value);
@@ -285,7 +285,7 @@ namespace ModelOptimizerNNEHelper
 		uint32 ExpectedDataSize = TensorShape.Volume() * UE::NNE::GetTensorDataTypeSizeInBytes(DataType);
 		if (ExpectedDataSize != DataView.NumBytes())
 		{
-			UE_LOG(LogNNE, Warning, TEXT("Tensor data size %i doesn't match expected data size %i"), 
+			UE_LOG(LogNNERuntimeRDGHlsl, Warning, TEXT("Tensor data size %i doesn't match expected data size %i"), 
 				   					DataView.NumBytes(), ExpectedDataSize);
 			return false;
 		}
@@ -301,13 +301,13 @@ namespace ModelOptimizerNNEHelper
 		const bool result = ModelProto.ParseFromArray(ONNXData.GetData(), ONNXData.Num());
 		if (!result)
 		{
-			UE_LOG(LogNNE, Warning, TEXT("Could not parse the input model as a ModelProto."));
+			UE_LOG(LogNNERuntimeRDGHlsl, Warning, TEXT("Could not parse the input model as a ModelProto."));
 			return false;
 		}
 
 		if (ModelProto.opset_import_size() < 1)
 		{
-			UE_LOG(LogNNE, Warning, TEXT("Could not read opset version from ONNX."));
+			UE_LOG(LogNNERuntimeRDGHlsl, Warning, TEXT("Could not read opset version from ONNX."));
 			return false;
 		}
 
@@ -355,20 +355,20 @@ namespace ModelOptimizerNNEHelper
 
 				if (!GetTensorInfoFromONNXInitializer(*Initializer, InitializerShape, InitializerDataType, Data))
 				{
-					UE_LOG(LogNNE, Error, TEXT("Tensor data could not be loaded for weights of output node '%s'"), ANSI_TO_TCHAR(Output.name().c_str()));
+					UE_LOG(LogNNERuntimeRDGHlsl, Error, TEXT("Tensor data could not be loaded for weights of output node '%s'"), ANSI_TO_TCHAR(Output.name().c_str()));
 					return false;
 				}
 
 				if (DataType != InitializerDataType)
 				{
-					UE_LOG(LogNNE, Warning, TEXT("Initializer type does not match output type for output tensor %s."), ANSI_TO_TCHAR(Output.name().c_str()));
+					UE_LOG(LogNNERuntimeRDGHlsl, Warning, TEXT("Initializer type does not match output type for output tensor %s."), ANSI_TO_TCHAR(Output.name().c_str()));
 					return false;
 				}
 
 
 				if (Shape != InitializerShape)
 				{
-					UE_LOG(LogNNE, Warning, TEXT("Initializer shape does not match output shape for output tensor %s."), ANSI_TO_TCHAR(Output.name().c_str()));
+					UE_LOG(LogNNERuntimeRDGHlsl, Warning, TEXT("Initializer shape does not match output shape for output tensor %s."), ANSI_TO_TCHAR(Output.name().c_str()));
 					return false;
 				}
 
@@ -425,7 +425,7 @@ namespace ModelOptimizerNNEHelper
 						int64 ValueClamped = FMath::Clamp<int64>(Value64, MIN_int32, MAX_int32);
 						if (ValueClamped != Value64)
 						{
-							UE_LOG(LogNNE, Display, TEXT("Overflow detected when converting to int32 attribute '%s' in node '%s'"), *AttributeName, *NNEOpType);
+							UE_LOG(LogNNERuntimeRDGHlsl, Display, TEXT("Overflow detected when converting to int32 attribute '%s' in node '%s'"), *AttributeName, *NNEOpType);
 						}
 						
 						Values.Add((int32)ValueClamped);
@@ -466,7 +466,7 @@ namespace ModelOptimizerNNEHelper
 					FAttributeTensor Tensor;
 					if (!GetAttributeTensorFromONNXInitializer(Attribute.t(), Tensor))
 					{
-						UE_LOG(LogNNE, Error, TEXT("Failed to get data from attribute %s in operator %s"), *AttributeName, *NNEOpName);
+						UE_LOG(LogNNERuntimeRDGHlsl, Error, TEXT("Failed to get data from attribute %s in operator %s"), *AttributeName, *NNEOpName);
 						return false;
 					}
 					Builder->AddOperatorAttribute(Op, AttributeName, FNNEAttributeValue(Tensor));
@@ -483,7 +483,7 @@ namespace ModelOptimizerNNEHelper
 						FAttributeTensor Tensor;
 						if (!GetAttributeTensorFromONNXInitializer(Initializer, Tensor))
 						{
-							UE_LOG(LogNNE, Error, TEXT("Failed to get data from tensor at index %i of attribute %s in operator %s"), TensorIndex, *AttributeName, *NNEOpName);
+							UE_LOG(LogNNERuntimeRDGHlsl, Error, TEXT("Failed to get data from tensor at index %i of attribute %s in operator %s"), TensorIndex, *AttributeName, *NNEOpName);
 							return false;
 						}
 						Values.Add(Tensor);
@@ -494,7 +494,7 @@ namespace ModelOptimizerNNEHelper
 				else
 				{
 					//Note: Would be good to have better error reporting by adding type (example: sparse tensor)
-					UE_LOG(LogNNE, Warning, TEXT("Unsupported attribute type for attribute '%s' in node '%s' of type '%s'"), *AttributeName, *NNEOpName, *NNEOpType);
+					UE_LOG(LogNNERuntimeRDGHlsl, Warning, TEXT("Unsupported attribute type for attribute '%s' in node '%s' of type '%s'"), *AttributeName, *NNEOpName, *NNEOpType);
 				}
 			}
 
@@ -511,7 +511,7 @@ namespace ModelOptimizerNNEHelper
 				{
 					if (!GetTensorInfoFromONNXInitializer(*Initializer, Shape, DataType, Data))
 					{
-						UE_LOG(LogNNE, Error, TEXT("Tensor data could not be loaded for weight '%s' in node '%s' of type '%s'"), ANSI_TO_TCHAR(TensorName.c_str()), *NNEOpName, *NNEOpType);
+						UE_LOG(LogNNERuntimeRDGHlsl, Error, TEXT("Tensor data could not be loaded for weight '%s' in node '%s' of type '%s'"), ANSI_TO_TCHAR(TensorName.c_str()), *NNEOpName, *NNEOpType);
 						return false;
 					}
 					TConstArrayView<uint8> DataView = Data.GetArrayView();
@@ -522,7 +522,7 @@ namespace ModelOptimizerNNEHelper
 					const onnx::ValueInfoProto* ValueInfoProto = GetValueInfoProtoFromGraphProto(Graph, TensorName);
 					if (!ValueInfoProto)
 					{
-						UE_LOG(LogNNE, Error, TEXT("Could not find Tensor ValueInfoProto or Initializer in graph for input '%s' in node '%s' of type '%s'"), ANSI_TO_TCHAR(TensorName.c_str()), *NNEOpName, *NNEOpType);
+						UE_LOG(LogNNERuntimeRDGHlsl, Error, TEXT("Could not find Tensor ValueInfoProto or Initializer in graph for input '%s' in node '%s' of type '%s'"), ANSI_TO_TCHAR(TensorName.c_str()), *NNEOpName, *NNEOpType);
 						return false;
 					}
 
@@ -542,7 +542,7 @@ namespace ModelOptimizerNNEHelper
 				const onnx::ValueInfoProto* ValueInfoProto = GetValueInfoProtoFromGraphProto(Graph, TensorName);
 				if (!ValueInfoProto)
 				{
-					UE_LOG(LogNNE, Error, TEXT("Could not find Tensor ValueInfoProto in graph for output '%s' in node '%s' of type '%s'"), ANSI_TO_TCHAR(TensorName.c_str()), *NNEOpName, *NNEOpType);
+					UE_LOG(LogNNERuntimeRDGHlsl, Error, TEXT("Could not find Tensor ValueInfoProto in graph for output '%s' in node '%s' of type '%s'"), ANSI_TO_TCHAR(TensorName.c_str()), *NNEOpName, *NNEOpType);
 					return false;
 				}
 				ENNETensorDataType DataType;
@@ -569,13 +569,13 @@ bool FModelOptimizerONNXToNNERT::Optimize(const FNNEModelRaw& InputModel, FNNEMo
 
 	if (!ModelOptimizerONNXToONNX.Optimize(InputModel, OptimizedONNXModel, Options))
 	{
-		UE_LOG(LogNNE, Warning, TEXT("Error while optimizing the ONNX model before convertion to NNERT format."));
+		UE_LOG(LogNNERuntimeRDGHlsl, Warning, TEXT("Error while optimizing the ONNX model before convertion to NNERT format."));
 		return false;
 	}
 
 	if (!ModelOptimizerNNEHelper::BuildNNEFormatFromONNX(OptimizedONNXModel.Data, OptimizedModel.Data))
 	{
-		UE_LOG(LogNNE, Warning, TEXT("Error while building NNERT Model from ONNX."));
+		UE_LOG(LogNNERuntimeRDGHlsl, Warning, TEXT("Error while building NNERT Model from ONNX."));
 		return false;
 	}
 	OptimizedModel.Format = ENNEInferenceFormat::NNERT;
diff --git a/Engine/Plugins/Experimental/NNERuntimeRDG/Source/NNERuntimeRDGUtils/Private/NNERuntimeRDGUtilsModelOptimizerONNX.cpp b/Engine/Plugins/Experimental/NNERuntimeRDG/Source/NNERuntimeRDGUtils/Private/NNERuntimeRDGUtilsModelOptimizerONNX.cpp
--- a/Engine/Plugins/Experimental/NNERuntimeRDG/Source/NNERuntimeRDGUtils/Private/NNERuntimeRDGUtilsModelOptimizerONNX.cpp
+++ b/Engine/Plugins/Experimental/NNERuntimeRDG/Source/NNERuntimeRDGUtils/Private/NNERuntimeRDGUtilsModelOptimizerONNX.cpp
@@ -7,7 +7,7 @@
 #include "HAL/PlatformTime.h"
 #include "Misc/FileHelper.h"
 #include "Misc/Paths.h"
-#include "NNE.h"
+#include "NNEHlslShadersLog.h"
 #include "NNEOnnxruntimeEditor.h"
 
 THIRD_PARTY_INCLUDES_START
@@ -37,7 +37,7 @@ public:
 	{
 		if (Model.Format != ENNEInferenceFormat::ONNX)
 		{
-			UE_LOG(LogNNE, Warning, TEXT("%s is expecting a model in ONNX format but received %u."), *(GetName()), Model.Format);
+			UE_LOG(LogNNERuntimeRDGHlsl, Warning, TEXT("%s is expecting a model in ONNX format but received %u."), *(GetName()), Model.Format);
 			return false;
 		}
 
@@ -45,7 +45,7 @@ public:
 		const bool result = ModelProto.ParseFromArray(Model.Data.GetData(), Model.Data.Num());
 		if (!result)
 		{
-			UE_LOG(LogNNE, Warning, TEXT("%s could not parse the input model as a ModelProto."), *(GetName()));
+			UE_LOG(LogNNERuntimeRDGHlsl, Warning, TEXT("%s could not parse the input model as a ModelProto."), *(GetName()));
 			return false;
 		}
 
@@ -156,7 +156,7 @@ public:
 	{
 		if (Model.Format != ENNEInferenceFormat::ONNX)
 		{
-			UE_LOG(LogNNE, Warning, TEXT("%s is expecting a model in ONNX format but received %u."), *(GetName()), Model.Format);
+			UE_LOG(LogNNERuntimeRDGHlsl, Warning, TEXT("%s is expecting a model in ONNX format but received %u."), *(GetName()), Model.Format);
 			return false;
 		}
 
@@ -164,7 +164,7 @@ public:
 		const bool result = ModelProto.ParseFromArray(Model.Data.GetData(), Model.Data.Num());
 		if (!result)
 		{
-			UE_LOG(LogNNE, Warning, TEXT("%s could not parse the input model as a ModelProto."), *(GetName()));
+			UE_LOG(LogNNERuntimeRDGHlsl, Warning, TEXT("%s could not parse the input model as a ModelProto."), *(GetName()));
 			return false;
 		}
 
@@ -214,7 +214,7 @@ public:
 	{
 		if (Model.Format != ENNEInferenceFormat::ONNX)
 		{
-			UE_LOG(LogNNE, Warning, TEXT("%s is expecting a mNNEdel in ONNX format but received %u."), *(GetName()), Model.Format);
+			UE_LOG(LogNNERuntimeRDGHlsl, Warning, TEXT("%s is expecting a mNNEdel in ONNX format but received %u."), *(GetName()), Model.Format);
 			return false;
 		}
 
@@ -222,12 +222,12 @@ public:
 		const bool result = ModelProto.ParseFromArray(Model.Data.GetData(), Model.Data.Num());
 		if (!result)
 		{
-			UE_LOG(LogNNE, Warning, TEXT("%s could not parse the input model as a ModelProto."), *(GetName()));
+			UE_LOG(LogNNERuntimeRDGHlsl, Warning, TEXT("%s could not parse the input model as a ModelProto."), *(GetName()));
 			return false;
 		}
 
 #ifdef ONNX_NO_EXCEPTIONS
-		UE_LOG(LogNNE, Warning, TEXT("ONNX Shape inference can't be run as exception are disabled."));
+		UE_LOG(LogNNERuntimeRDGHlsl, Warning, TEXT("ONNX Shape inference can't be run as exception are disabled."));
 		return true;
 #else
 
@@ -239,7 +239,7 @@ public:
 		}
 		catch (onnx::InferenceError& e)
 		{
-			UE_LOG(LogNNE, Warning, TEXT("Shape inference failed with : %s."), ANSI_TO_TCHAR(e.what()));
+			UE_LOG(LogNNERuntimeRDGHlsl, Warning, TEXT("Shape inference failed with : %s."), ANSI_TO_TCHAR(e.what()));
 		}
 #endif
 		
diff --git a/Engine/Plugins/Experimental/NNERuntimeRDG/Source/NNERuntimeRDGUtils/Private/NNERuntimeRDGUtilsModule.cpp b/Engine/Plugins/Experimental/NNERuntimeRDG/Source/NNERuntimeRDGUtils/Private/NNERuntimeRDGUtilsModule.cpp
--- a/Engine/Plugins/Experimental/NNERuntimeRDG/Source/NNERuntimeRDGUtils/Private/NNERuntimeRDGUtilsModule.cpp
+++ b/Engine/Plugins/Experimental/NNERuntimeRDG/Source/NNERuntimeRDGUtils/Private/NNERuntimeRDGUtilsModule.cpp
@@ -1,9 +1,9 @@
 // Copyright Epic Games, Inc. All Rights Reserved.
 
-#include "Modules/ModuleManager.h"
 #include "Interfaces/IPluginManager.h"
 #include "Misc/Paths.h"
-#include "NNE.h"
+#include "Modules/ModuleManager.h"
+#include "NNEHlslShadersLog.h"
 #include "NNEOnnxruntimeEditor.h"
 
 class FNNERuntimeRDGUtilsModule : public IModuleInterface
@@ -18,14 +18,14 @@ public:
 		OrtDllHandle = FPlatformProcess::GetDllHandle(*OrtSharedLibPath);
 		if (!OrtDllHandle)
 		{
-			UE_LOG(LogNNE, Fatal, TEXT("Failed to load ONNX Runtime shared library!"));
+			UE_LOG(LogNNERuntimeRDGHlsl, Fatal, TEXT("Failed to load ONNX Runtime shared library!"));
 			return;
 		}
 
 		TUniquePtr<UE::NNEOnnxruntime::OrtApiFunctions> OrtApiFunctions = UE::NNEOnnxruntime::LoadApiFunctions(OrtDllHandle);
 		if (!OrtApiFunctions.IsValid())
 		{
-			UE_LOG(LogNNE, Fatal, TEXT("Failed to load ONNX Runtime shared library functions!"));
+			UE_LOG(LogNNERuntimeRDGHlsl, Fatal, TEXT("Failed to load ONNX Runtime shared library functions!"));
 			return;
 		}
 
