[NNE] NNERuntimeRDGHLSL change float16 data type to 'half', since 'float16_t' is not supported by the Metal Shading Language Specification.

Problem:
HLSL Shaders using 'float16_t' type do not compile on platforms with Metal Shading Language, because of missing support.

Solution:
Use 'half' type instead on Metal platforms.

Tests:
Compile all HLSL shaders on Metal platform.

#tests Editor, Mac compile all shaders and run all NNE automation tests

#rnx

diff --git a/Engine/Plugins/Experimental/NNERuntimeRDG/Source/NNEHlslShaders/Private/NNEHlslShadersTypeHelper.cpp b/Engine/Plugins/Experimental/NNERuntimeRDG/Source/NNEHlslShaders/Private/NNEHlslShadersTypeHelper.cpp
--- a/Engine/Plugins/Experimental/NNERuntimeRDG/Source/NNEHlslShaders/Private/NNEHlslShadersTypeHelper.cpp
+++ b/Engine/Plugins/Experimental/NNERuntimeRDG/Source/NNEHlslShaders/Private/NNEHlslShadersTypeHelper.cpp
@@ -9,7 +9,11 @@ namespace UE::NNEHlslShaders::Internal
 		switch (ShaderType)
 		{
 			case ENNEShaderDataType::BOOL	 : return TEXT("bool");
+#if PLATFORM_MAC && HAS_METAL
+			case ENNEShaderDataType::FLOAT16 : return TEXT("half");
+#else
 			case ENNEShaderDataType::FLOAT16 : return TEXT("float16_t");
+#endif
 			case ENNEShaderDataType::FLOAT32 : return TEXT("float");
 			case ENNEShaderDataType::INT8	 : return TEXT("int8_t");
 			case ENNEShaderDataType::INT16	 : return TEXT("int16_t");
