[NNE] Disable OPs failing tests on Mac

Problems:
* Some OPs are not working properly on Mac (failling regression tests)

Solutions:
* Disable their registration on Mac so that users are aware when input model requires such OPs

Tests:
* Win64 Editor NNEQA: all
* Mac Editor NNEQA: all

#rnx

diff --git a/Engine/Plugins/Experimental/NNERuntimeRDG/Source/NNERuntimeRDG/Private/Hlsl/NNERuntimeRDGCumSum.cpp b/Engine/Plugins/Experimental/NNERuntimeRDG/Source/NNERuntimeRDG/Private/Hlsl/NNERuntimeRDGCumSum.cpp
--- a/Engine/Plugins/Experimental/NNERuntimeRDG/Source/NNERuntimeRDG/Private/Hlsl/NNERuntimeRDGCumSum.cpp
+++ b/Engine/Plugins/Experimental/NNERuntimeRDG/Source/NNERuntimeRDG/Private/Hlsl/NNERuntimeRDGCumSum.cpp
@@ -283,9 +283,12 @@ namespace UE::NNERuntimeRDG::Private::Hlsl
 	{
 		using namespace UE::NNEHlslShaders::Internal;
 
-		Registry.OpAdd({ {TEXT("CumSum"), TEXT("Onnx")}, 11}, CreateCumSumOperator<11>, ValidateCumSumOperator<11>);
-		Registry.OpAdd({ {TEXT("CumSum"), TEXT("Onnx")}, 14}, CreateCumSumOperator<14>, ValidateCumSumOperator<14>);
-		
+		// Note: CumSum is currently not working on Mac.
+		#if !PLATFORM_MAC
+			Registry.OpAdd({ {TEXT("CumSum"), TEXT("Onnx")}, 11}, CreateCumSumOperator<11>, ValidateCumSumOperator<11>);
+			Registry.OpAdd({ {TEXT("CumSum"), TEXT("Onnx")}, 14}, CreateCumSumOperator<14>, ValidateCumSumOperator<14>);
+		#endif
+
 		return true;
 	}
 } // UE::NNERuntimeRDG::Private::Hlsl
diff --git a/Engine/Plugins/Experimental/NNERuntimeRDG/Source/NNERuntimeRDG/Private/Hlsl/NNERuntimeRDGResize.cpp b/Engine/Plugins/Experimental/NNERuntimeRDG/Source/NNERuntimeRDG/Private/Hlsl/NNERuntimeRDGResize.cpp
--- a/Engine/Plugins/Experimental/NNERuntimeRDG/Source/NNERuntimeRDG/Private/Hlsl/NNERuntimeRDGResize.cpp
+++ b/Engine/Plugins/Experimental/NNERuntimeRDG/Source/NNERuntimeRDG/Private/Hlsl/NNERuntimeRDGResize.cpp
@@ -414,11 +414,15 @@ namespace UE::NNERuntimeRDG::Private::Hlsl
 	bool RegisterResizeOperator(FOperatorRegistryHlsl& Registry)
 	{
 		// Note: support of a particular version is partial with respect to tensor data types (only the most typical ones are usually supported).
-		Registry.OpAdd({{TEXT("Resize"), TEXT("Onnx")}, 10}, CreateResizeOperator, ValidateResizeOperator);
-		Registry.OpAdd({{TEXT("Resize"), TEXT("Onnx")}, 11}, CreateResizeOperator, ValidateResizeOperator);
-		Registry.OpAdd({{TEXT("Resize"), TEXT("Onnx")}, 13}, CreateResizeOperator, ValidateResizeOperator);
-		Registry.OpAdd({{TEXT("Resize"), TEXT("Onnx")}, 18}, CreateResizeOperator, ValidateResizeOperator);
-		Registry.OpAdd({{TEXT("Resize"), TEXT("Onnx")}, 19}, CreateResizeOperator, ValidateResizeOperator);
+		
+		// Note: CumSum is currently not working on Mac.
+		#if !PLATFORM_MAC
+			Registry.OpAdd({{TEXT("Resize"), TEXT("Onnx")}, 10}, CreateResizeOperator, ValidateResizeOperator);
+			Registry.OpAdd({{TEXT("Resize"), TEXT("Onnx")}, 11}, CreateResizeOperator, ValidateResizeOperator);
+			Registry.OpAdd({{TEXT("Resize"), TEXT("Onnx")}, 13}, CreateResizeOperator, ValidateResizeOperator);
+			Registry.OpAdd({{TEXT("Resize"), TEXT("Onnx")}, 18}, CreateResizeOperator, ValidateResizeOperator);
+			Registry.OpAdd({{TEXT("Resize"), TEXT("Onnx")}, 19}, CreateResizeOperator, ValidateResizeOperator);
+		#endif
 		return true;
 	}
 } // UE::NNERuntimeRDG::Private::Hlsl
