[NNE] Fix Hlsl CumSum compilation for Mac but disable it

Problems:
* TInitCumSum shaders breaks compilation on Mac because Metal compiler can't translate struct assignment.

Solutions:
* Assign values directly to struct fields in TInitCumSum for Mac

Tests:
* Win64 NNEQA: all
* Mac NNEQA: all

#rnx

diff --git a/Engine/Plugins/Experimental/NNERuntimeRDG/Shaders/Private/NNEHlslShaders/NNEHlslShadersCumSum.usf b/Engine/Plugins/Experimental/NNERuntimeRDG/Shaders/Private/NNEHlslShaders/NNEHlslShadersCumSum.usf
--- a/Engine/Plugins/Experimental/NNERuntimeRDG/Shaders/Private/NNEHlslShaders/NNEHlslShadersCumSum.usf
+++ b/Engine/Plugins/Experimental/NNERuntimeRDG/Shaders/Private/NNEHlslShaders/NNEHlslShadersCumSum.usf
@@ -52,7 +52,13 @@ void InitCumSum(uint3 DispatchThreadID : SV_DispatchThreadID)
 
 	for (uint Idx = DispatchThreadID.x; Idx < NumThreadGroupsPerScan * NumThreadGroupsY * NumThreadGroupsZ; Idx += Increment)
 	{
-		FPartitionDescriptor InitPD;
+#if METAL_SM6_PROFILE
+		PartitionDescriptor[Idx].StatusFlag = STATUS_INVALID;
+		PartitionDescriptor[Idx].InclusivePrefix = (WORK_TYPE) 0;
+		PartitionDescriptor[Idx].Aggregate = (WORK_TYPE) 0;
+		PartitionDescriptor[Idx].PadToQWord = 0;
+#else
+        FPartitionDescriptor InitPD;
 
 		InitPD.StatusFlag = STATUS_INVALID;
 		InitPD.InclusivePrefix = (WORK_TYPE) 0;
@@ -60,6 +66,7 @@ void InitCumSum(uint3 DispatchThreadID : SV_DispatchThreadID)
 		InitPD.PadToQWord = 0;
 
 		PartitionDescriptor[Idx] = InitPD;
+#endif
 
 		if (Idx % NumThreadGroupsPerScan == 0)
 		{
