[NNE] HLSL Return error with context in CanCreateModel data

Problem: CreateModel return warning instead of error but not context and CanCreateModelData do not catch all failure case.
Solution: Factorize the code so both methods handle the same cases while only CreateModelData log (as error)

#test: NNE smoke tests: in win64 editor debug + win64 standalone dev.

#rnx

diff --git a/Engine/Plugins/Experimental/NNERuntimeRDG/Source/NNERuntimeRDG/Private/NNERuntimeRDGHlsl.cpp b/Engine/Plugins/Experimental/NNERuntimeRDG/Source/NNERuntimeRDG/Private/NNERuntimeRDGHlsl.cpp
--- a/Engine/Plugins/Experimental/NNERuntimeRDG/Source/NNERuntimeRDG/Private/NNERuntimeRDGHlsl.cpp
+++ b/Engine/Plugins/Experimental/NNERuntimeRDG/Source/NNERuntimeRDG/Private/NNERuntimeRDGHlsl.cpp
@@ -101,21 +101,70 @@ bool UNNERuntimeRDGHlslImpl::Init()
 	return true;
 }
 
-namespace ConsoleCommands
+namespace UE::NNERuntimeRDG::Private::Hlsl
 {
-	static FAutoConsoleCommand GetAutomationRuntimeFilterCommand(
-		TEXT("nne.hlsl.getoperatorsupportmatrix"), TEXT("Get the NNERuntimeRDGHlsl operators support matrix in term of ONNX."),
-		FConsoleCommandWithArgsDelegate::CreateStatic(
-			[](const TArray< FString >& Args)
+	namespace ConsoleCommands
+	{
+		static FAutoConsoleCommand GetAutomationRuntimeFilterCommand(
+			TEXT("nne.hlsl.getoperatorsupportmatrix"), TEXT("Get the NNERuntimeRDGHlsl operators support matrix in term of ONNX."),
+			FConsoleCommandWithArgsDelegate::CreateStatic(
+				[](const TArray< FString >& Args)
+				{
+					FOperatorRegistryHlsl* Registry = FOperatorRegistryHlsl::Get();
+					check(Registry != nullptr);
+					FString SupportMatrix = Registry->ListAllRegisteredOperators();
+					UE_LOG(LogNNERuntimeRDGHlsl, Display, TEXT("Operators support matrix: \n%s"), *SupportMatrix);
+				}
+			)
+		);
+	} // namespace ConsoleCommands
+
+	namespace Details
+	{
+		UNNERuntimeRDGHlslImpl::ECanCreateModelDataStatus CheckCanCreateModelData(bool bShouldLog, const FString& FileType, TConstArrayView64<uint8> FileData, const TMap<FString, TConstArrayView64<uint8>>& AdditionalFileData, const FGuid& FileId, const ITargetPlatform* TargetPlatform)
+		{
+#ifdef NNE_UTILITIES_AVAILABLE
+			if (FileType.Compare("onnx", ESearchCase::IgnoreCase) != 0)
+			{
+				if (bShouldLog)
+				{
+					UE_LOG(LogNNERuntimeRDGHlsl, Error, TEXT("Cannot create the model data with id %s (Filetype: %s), Only 'onnx' file type is supported"), *FileId.ToString(EGuidFormats::Digits).ToLower(), *FileType);
+				}
+				return UNNERuntimeRDGHlslImpl::ECanCreateModelDataStatus::Fail;
+			}
+
+			// Check model is not > 2GB
+			if ((TArray<uint8>::SizeType)FileData.Num() != FileData.Num())
+			{
+				if (bShouldLog)
+				{
+					UE_LOG(LogNNERuntimeRDGHlsl, Error, TEXT("Cannot create the model data with id %s (Filetype: %s), models > 2GBs are not supported"), *FileId.ToString(EGuidFormats::Digits).ToLower(), *FileType);
+				}
+				return UNNERuntimeRDGHlslImpl::ECanCreateModelDataStatus::Fail;
+			}
+
+			if (!AdditionalFileData.IsEmpty())
 			{
-				FOperatorRegistryHlsl* Registry = FOperatorRegistryHlsl::Get();
-				check(Registry != nullptr);
-				FString SupportMatrix = Registry->ListAllRegisteredOperators();
-				UE_LOG(LogNNERuntimeRDGHlsl, Display, TEXT("Operators support matrix: \n%s"), *SupportMatrix);
+				if (bShouldLog)
+				{
+					UE_LOG(LogNNERuntimeRDGHlsl, Error, TEXT("Cannot create the model data with id %s (Filetype: %s), external data not supported at the moment, please convert the model to internal storage. See https://onnx.ai/onnx/repo-docs/ExternalData.html"), *FileId.ToString(EGuidFormats::Digits).ToLower(), *FileType);
+				}
+				return UNNERuntimeRDGHlslImpl::ECanCreateModelDataStatus::Fail;
 			}
-		)
-	);
-} // ConsoleCommands
+
+			return UNNERuntimeRDGHlslImpl::ECanCreateModelDataStatus::Ok;
+#else
+			if (bShouldLog)
+			{
+				UE_LOG(LogNNERuntimeRDGHlsl, Error, TEXT("Cannot create the model data with id %s (Filetype: %s), NNERuntimeRDGUtils is not available on this platform"), *FileId.ToString(EGuidFormats::Digits).ToLower(), *FileType);
+			}
+			return UNNERuntimeRDGHlslImpl::ECanCreateModelDataStatus::Fail;
+#endif
+		}
+	} // namespace Details
+
+} // namespace UE::NNERuntimeRDG::Private::Hlsl
+
 
 bool UNNERuntimeRDGHlslImpl::IsCurrentPlatformSupported()
 {
@@ -146,12 +195,7 @@ bool UNNERuntimeRDGHlslImpl::IsCurrentPlatformSupported()
 
 UNNERuntimeRDGHlslImpl::ECanCreateModelDataStatus UNNERuntimeRDGHlslImpl::CanCreateModelData(const FString& FileType, TConstArrayView64<uint8> FileData, const TMap<FString, TConstArrayView64<uint8>>& AdditionalFileData, const FGuid& FileId, const ITargetPlatform* TargetPlatform) const
 {
-#ifdef NNE_UTILITIES_AVAILABLE
-	return FileType.Compare("onnx", ESearchCase::IgnoreCase) == 0 ? ECanCreateModelDataStatus::Ok : ECanCreateModelDataStatus::FailFileIdNotSupported;
-#else
-	UE_LOG(LogNNERuntimeRDGHlsl, Display, TEXT("NNERuntimeRDGUtils is not available on this platform"));
-	return ECanCreateModelDataStatus::Fail;
-#endif
+	return Details::CheckCanCreateModelData(/*bShouldLog*/ false , FileType, FileData, AdditionalFileData, FileId, TargetPlatform);
 }
 
 UNNERuntimeRDGHlslImpl::ECanCreateModelRDGStatus UNNERuntimeRDGHlslImpl::CanCreateModelRDG(const TObjectPtr<UNNEModelData> ModelData) const
@@ -181,22 +225,8 @@ UNNERuntimeRDGHlslImpl::ECanCreateModelRDGStatus UNNERuntimeRDGHlslImpl::CanCrea
 
 TSharedPtr<UE::NNE::FSharedModelData> UNNERuntimeRDGHlslImpl::CreateModelData(const FString& FileType, TConstArrayView64<uint8> FileData, const TMap<FString, TConstArrayView64<uint8>>& AdditionalFileData, const FGuid& FileId, const ITargetPlatform* TargetPlatform)
 {
-	// Check model is not > 2GB
-	if ((TArray<uint8>::SizeType)FileData.Num() != FileData.Num())
-	{
-		UE_LOG(LogNNERuntimeRDGHlsl, Warning, TEXT("Cannot create the model data with id %s (Filetype: %s), models > 2GBs are not supported"), *FileId.ToString(EGuidFormats::Digits).ToLower(), *FileType);
-		return {};
-	}
-
-	if (!AdditionalFileData.IsEmpty())
-	{
-		UE_LOG(LogNNERuntimeRDGHlsl, Warning, TEXT("Cannot create the model data with id %s (Filetype: %s), external data not supported at the moment, please convert the model to internal storage. See https://onnx.ai/onnx/repo-docs/ExternalData.html"), *FileId.ToString(EGuidFormats::Digits).ToLower(), *FileType);
-		return {};
-	}
-
-	if (CanCreateModelData(FileType, FileData, AdditionalFileData, FileId, TargetPlatform) != ECanCreateModelDataStatus::Ok)
+	if (Details::CheckCanCreateModelData(/*bShouldLog*/ true, FileType, FileData, AdditionalFileData, FileId, TargetPlatform) != ECanCreateModelDataStatus::Ok)
 	{
-		UE_LOG(LogNNERuntimeRDGHlsl, Warning, TEXT("Cannot create the model data with id %s (Filetype: %s)"), *FileId.ToString(EGuidFormats::Digits).ToLower(), *FileType);
 		return {};
 	}
 
