[NNE] Disable Conv Winograd on Mac.

Problem:
On Mac shader can not be compiled.

Solution:
Disable on Mac for now.

#tests Start Editor and verify all shaders compile and run all runtime HLSL tests

#rnx

diff --git a/Engine/Plugins/Experimental/NNERuntimeRDG/Source/NNEHlslShaders/Private/NNEHlslShadersConvWinogradMMMCS.cpp b/Engine/Plugins/Experimental/NNERuntimeRDG/Source/NNEHlslShaders/Private/NNEHlslShadersConvWinogradMMMCS.cpp
--- a/Engine/Plugins/Experimental/NNERuntimeRDG/Source/NNEHlslShaders/Private/NNEHlslShadersConvWinogradMMMCS.cpp
+++ b/Engine/Plugins/Experimental/NNERuntimeRDG/Source/NNEHlslShaders/Private/NNEHlslShadersConvWinogradMMMCS.cpp
@@ -11,6 +11,9 @@ namespace UE::NNEHlslShaders::Internal
 {
 	bool FConvWinogradMMMCS::ShouldCompilePermutation(const FGlobalShaderPermutationParameters& Parameters)
 	{
+#if PLATFORM_MAC
+		return false;
+#else
 		if (!FHlslShaderBase::ShouldCompilePermutation(Parameters))
 		{
 			return false;
@@ -18,6 +21,7 @@ namespace UE::NNEHlslShaders::Internal
 		FPermutationDomain PermutationVector(Parameters.PermutationId);
 		ENNEShaderDataType DataType = PermutationVector.Get<FConvWinogradMMMCS::FDataType>();
 		return DataType == ENNEShaderDataType::FLOAT16 || DataType == ENNEShaderDataType::FLOAT32;
+#endif // PLATFORM_MAC
 	}
 
 	void FConvWinogradMMMCS::ModifyCompilationEnvironment(const FGlobalShaderPermutationParameters& InParameters, FShaderCompilerEnvironment& OutEnvironment)
diff --git a/Engine/Plugins/Experimental/NNERuntimeRDG/Source/NNEHlslShaders/Private/NNEHlslShadersTypeHelper.cpp b/Engine/Plugins/Experimental/NNERuntimeRDG/Source/NNEHlslShaders/Private/NNEHlslShadersTypeHelper.cpp
--- a/Engine/Plugins/Experimental/NNERuntimeRDG/Source/NNEHlslShaders/Private/NNEHlslShadersTypeHelper.cpp
+++ b/Engine/Plugins/Experimental/NNERuntimeRDG/Source/NNEHlslShaders/Private/NNEHlslShadersTypeHelper.cpp
@@ -9,11 +9,7 @@ namespace UE::NNEHlslShaders::Internal
 		switch (ShaderType)
 		{
 			case ENNEShaderDataType::BOOL	 : return TEXT("bool");
-#if PLATFORM_MAC && HAS_METAL
-			case ENNEShaderDataType::FLOAT16 : return TEXT("half");
-#else
 			case ENNEShaderDataType::FLOAT16 : return TEXT("float16_t");
-#endif
 			case ENNEShaderDataType::FLOAT32 : return TEXT("float");
 			case ENNEShaderDataType::INT8	 : return TEXT("int8_t");
 			case ENNEShaderDataType::INT16	 : return TEXT("int16_t");
diff --git a/Engine/Plugins/Experimental/NNERuntimeRDG/Source/NNERuntimeRDG/Private/Hlsl/NNERuntimeRDGConv.cpp b/Engine/Plugins/Experimental/NNERuntimeRDG/Source/NNERuntimeRDG/Private/Hlsl/NNERuntimeRDGConv.cpp
--- a/Engine/Plugins/Experimental/NNERuntimeRDG/Source/NNERuntimeRDG/Private/Hlsl/NNERuntimeRDGConv.cpp
+++ b/Engine/Plugins/Experimental/NNERuntimeRDG/Source/NNERuntimeRDG/Private/Hlsl/NNERuntimeRDGConv.cpp
@@ -534,10 +534,12 @@ namespace UE::NNERuntimeRDG::Private::Hlsl
 			check(Output.GetShape().Rank() == Input.GetShape().Rank());
 			check(NumDimensions == (Input.GetShape().Rank() - 2));
 
+#if !PLATFORM_MAC
 			if (DispatchConvWinograd(GraphBuilder, Input, Weights, Bias, Output))
 			{
 				return;
 			}
+#endif // !PLATFORM_MAC
 			if (DispatchConvMatmul(GraphBuilder, Input, Weights, Bias, Output))
 			{
 				return;
