[NNE] Fix winograd mmm shader compilation with bindless enabled

Problem: shader is using Buffer<vector<float,2>> style for resource declaration witch is incompatible with bindless atm.
Solution: use Buffer<float2> style of declaration instead.

Tests:
win64 editor debug + bindless for all shaders stage: all nne tests + style transfer
win64 editor debug + default bindless setup : all nne tests + style transfer
win64 packaged standalone debug + bindless for all shaders stage: all nne tests + style transfer
win64 packaged standalone debug + default bindless setup: all nne tests + style transfer
console packaged dev : nne smoketest + convolutions tests + style transfer

#rnx

diff --git a/Engine/Plugins/Experimental/NNERuntimeRDG/Shaders/Private/NNEHlslShaders/NNEHlslShadersConvWinogradMMM.usf b/Engine/Plugins/Experimental/NNERuntimeRDG/Shaders/Private/NNEHlslShaders/NNEHlslShadersConvWinogradMMM.usf
--- a/Engine/Plugins/Experimental/NNERuntimeRDG/Shaders/Private/NNEHlslShaders/NNEHlslShadersConvWinogradMMM.usf
+++ b/Engine/Plugins/Experimental/NNERuntimeRDG/Shaders/Private/NNEHlslShaders/NNEHlslShadersConvWinogradMMM.usf
@@ -17,7 +17,10 @@
 #define THREADGROUP_SIZE_Y (BLOCK_COUNT / THREADGROUP_SIZE_X)
 #define THREADGROUP_SIZE (THREADGROUP_SIZE_X * THREADGROUP_SIZE_Y)
 
-#define WORK_TYPE_VECTOR vector <WORK_TYPE, VECTOR_SIZE>
+// Winograd does not compile on D3D12 SM6 + bindless if using vector<x,y> atm
+// hence the define is directly set to the vectorized type in ModifyCompilationEnvironment()
+// rather than defined as a templated vector here
+//#define WORK_TYPE_VECTOR vector <WORK_TYPE, VECTOR_SIZE>
 
 Buffer<WORK_TYPE_VECTOR> Input; // Ni x 36 x K x M
 Buffer<WORK_TYPE_VECTOR> Weight; // 36 x K x N 
diff --git a/Engine/Plugins/Experimental/NNERuntimeRDG/Source/NNEHlslShaders/Private/NNEHlslShadersConvWinogradMMMCS.cpp b/Engine/Plugins/Experimental/NNERuntimeRDG/Source/NNEHlslShaders/Private/NNEHlslShadersConvWinogradMMMCS.cpp
--- a/Engine/Plugins/Experimental/NNERuntimeRDG/Source/NNEHlslShaders/Private/NNEHlslShadersConvWinogradMMMCS.cpp
+++ b/Engine/Plugins/Experimental/NNERuntimeRDG/Source/NNEHlslShaders/Private/NNEHlslShadersConvWinogradMMMCS.cpp
@@ -30,6 +30,21 @@ namespace UE::NNEHlslShaders::Internal
 		ENNEShaderDataType DataType = PermutationVector.Get<FConvWinogradMMMCS::FDataType>();
 		OutEnvironment.SetDefine(TEXT("WORK_TYPE"), ShaderDataTypeToName(DataType));
 
+		// Winograd does not compile on D3D12 SM6 + bindless if using vector<x,y> atm
+		// hence the define here directly set to the vectorized type rather as a define of a templated vector in the shader
+		if (DataType == ENNEShaderDataType::FLOAT16)
+		{
+			OutEnvironment.SetDefine(TEXT("WORK_TYPE_VECTOR"), TEXT("float16_t2"));
+		}
+		else if (DataType == ENNEShaderDataType::FLOAT32)
+		{
+			OutEnvironment.SetDefine(TEXT("WORK_TYPE_VECTOR"), TEXT("float2"));
+		}
+		else
+		{
+			OutEnvironment.SetDefine(TEXT("WORK_TYPE_VECTOR"), TEXT("UNSUPPORTED_DATA_TYPE"));
+		}
+
 		FGlobalShader::ModifyCompilationEnvironment(InParameters, OutEnvironment);
 		OutEnvironment.CompilerFlags.Add(CFLAG_AllowRealTypes);
 	}
