[NNE] Cleanup of internal dependancies between NNE plugins and NNEQA

Problems:
* NNEQAs have dependancies to Source/NNE/internal and NNERuntimeRDG
* NNE plugins have dependancies to Source/NNE/internal

Solutions (wip this is only the first PR toward the whole cleanup) :
* Moved the RDG unit test from NNEQA to NNERuntimeRDG and cleaned related dependancies
* Start to remove files from Source/NNE/internal and place them into plugins instead.

#test: Win64 debug editor: all nne tests including unit tests + all json based test reimported

#rnx

diff --git a/Engine/Plugins/Experimental/NNERuntimeRDG/Source/NNEHlslShaders/Internal/NNEHlslShadersBatchNormalizationCS.h b/Engine/Plugins/Experimental/NNERuntimeRDG/Source/NNEHlslShaders/Internal/NNEHlslShadersBatchNormalizationCS.h
--- a/Engine/Plugins/Experimental/NNERuntimeRDG/Source/NNEHlslShaders/Internal/NNEHlslShadersBatchNormalizationCS.h
+++ b/Engine/Plugins/Experimental/NNERuntimeRDG/Source/NNEHlslShaders/Internal/NNEHlslShadersBatchNormalizationCS.h
@@ -3,7 +3,6 @@
 #pragma once
 
 #include "NNEHlslShadersBase.h"
-#include "NNEOperator.h"
 #include "RenderGraphUtils.h"
 #include "ShaderParameterUtils.h"
 
diff --git a/Engine/Plugins/Experimental/NNERuntimeRDG/Source/NNEHlslShaders/Internal/NNEHlslShadersCastCS.h b/Engine/Plugins/Experimental/NNERuntimeRDG/Source/NNEHlslShaders/Internal/NNEHlslShadersCastCS.h
--- a/Engine/Plugins/Experimental/NNERuntimeRDG/Source/NNEHlslShaders/Internal/NNEHlslShadersCastCS.h
+++ b/Engine/Plugins/Experimental/NNERuntimeRDG/Source/NNEHlslShaders/Internal/NNEHlslShadersCastCS.h
@@ -4,7 +4,6 @@
 
 #include "NNEHlslShadersBase.h"
 #include "NNEHlslShadersTypeHelper.h"
-#include "NNEOperator.h"
 #include "RenderGraphUtils.h"
 #include "ShaderParameterUtils.h"
 
diff --git a/Engine/Plugins/Experimental/NNERuntimeRDG/Source/NNEHlslShaders/Internal/NNEHlslShadersElementWiseBinaryCS.h b/Engine/Plugins/Experimental/NNERuntimeRDG/Source/NNEHlslShaders/Internal/NNEHlslShadersElementWiseBinaryCS.h
--- a/Engine/Plugins/Experimental/NNERuntimeRDG/Source/NNEHlslShaders/Internal/NNEHlslShadersElementWiseBinaryCS.h
+++ b/Engine/Plugins/Experimental/NNERuntimeRDG/Source/NNEHlslShaders/Internal/NNEHlslShadersElementWiseBinaryCS.h
@@ -2,8 +2,8 @@
 
 #pragma once
 
-#include "NNEOperator.h"
 #include "NNEHlslShadersBase.h"
+#include "NNEHlslShadersOperator.h"
 #include "RenderGraphUtils.h"
 #include "ShaderParameterUtils.h"
 
@@ -22,7 +22,7 @@ namespace UE::NNEHlslShaders::Internal
 		DECLARE_GLOBAL_SHADER(TElementWiseBinaryCS);
 		SHADER_USE_PARAMETER_STRUCT(TElementWiseBinaryCS, FHlslShaderBase)
 
-		class FOperatorType : SHADER_PERMUTATION_ENUM_CLASS("OP_TYPENAME", NNE::Internal::EElementWiseBinaryOperatorType);
+		class FOperatorType : SHADER_PERMUTATION_ENUM_CLASS("OP_TYPENAME", EElementWiseBinaryOperatorType);
 		class FBinaryNumDimensions : SHADER_PERMUTATION_RANGE_INT("NUM_DIMENSIONS", 1, FElementWiseBinaryConstants::MAX_NUM_DIMENSIONS);
 		using FPermutationDomain = TShaderPermutationDomain<FOperatorType, FBinaryNumDimensions>;
 
@@ -41,6 +41,6 @@ namespace UE::NNEHlslShaders::Internal
 
 	private:
 
-		static const FString GetOpFunc(NNE::Internal::EElementWiseBinaryOperatorType OpType);
+		static const FString GetOpFunc(EElementWiseBinaryOperatorType OpType);
 	};
 } // UE::NNEHlslShaders::Internal
diff --git a/Engine/Plugins/Experimental/NNERuntimeRDG/Source/NNEHlslShaders/Internal/NNEHlslShadersElementWiseUnaryCS.h b/Engine/Plugins/Experimental/NNERuntimeRDG/Source/NNEHlslShaders/Internal/NNEHlslShadersElementWiseUnaryCS.h
--- a/Engine/Plugins/Experimental/NNERuntimeRDG/Source/NNEHlslShaders/Internal/NNEHlslShadersElementWiseUnaryCS.h
+++ b/Engine/Plugins/Experimental/NNERuntimeRDG/Source/NNEHlslShaders/Internal/NNEHlslShadersElementWiseUnaryCS.h
@@ -2,8 +2,8 @@
 
 #pragma once
 
-#include "NNEOperator.h"
 #include "NNEHlslShadersBase.h"
+#include "NNEHlslShadersOperator.h"
 #include "RenderGraphUtils.h"
 #include "ShaderParameterUtils.h"
 
@@ -21,7 +21,7 @@ namespace UE::NNEHlslShaders::Internal
 		DECLARE_GLOBAL_SHADER(TElementWiseUnaryCS);
 		SHADER_USE_PARAMETER_STRUCT(TElementWiseUnaryCS, FHlslShaderBase)
 
-		class FOperatorType : SHADER_PERMUTATION_ENUM_CLASS("OP_TYPENAME", NNE::Internal::EElementWiseUnaryOperatorType);
+		class FOperatorType : SHADER_PERMUTATION_ENUM_CLASS("OP_TYPENAME", EElementWiseUnaryOperatorType);
 		class FAlphaOnGPU : SHADER_PERMUTATION_BOOL("ALPHA_ON_GPU");
 		class FBetaOnGPU : SHADER_PERMUTATION_BOOL("BETA_ON_GPU");
 		using FPermutationDomain = TShaderPermutationDomain<FOperatorType,FAlphaOnGPU,FBetaOnGPU>;
@@ -46,6 +46,6 @@ namespace UE::NNEHlslShaders::Internal
 
 	private:
 
-		static const FString GetOpFunc(NNE::Internal::EElementWiseUnaryOperatorType OpType);
+		static const FString GetOpFunc(EElementWiseUnaryOperatorType OpType);
 	};
 } // UE::NNEHlslShaders::Internal
diff --git a/Engine/Plugins/Experimental/NNERuntimeRDG/Source/NNEHlslShaders/Internal/NNEHlslShadersElementWiseVariadicCS.h b/Engine/Plugins/Experimental/NNERuntimeRDG/Source/NNEHlslShaders/Internal/NNEHlslShadersElementWiseVariadicCS.h
--- a/Engine/Plugins/Experimental/NNERuntimeRDG/Source/NNEHlslShaders/Internal/NNEHlslShadersElementWiseVariadicCS.h
+++ b/Engine/Plugins/Experimental/NNERuntimeRDG/Source/NNEHlslShaders/Internal/NNEHlslShadersElementWiseVariadicCS.h
@@ -2,8 +2,8 @@
 
 #pragma once
 
-#include "NNEOperator.h"
 #include "NNEHlslShadersBase.h"
+#include "NNEHlslShadersOperator.h"
 #include "RenderGraphUtils.h"
 #include "ShaderParameterUtils.h"
 
@@ -23,7 +23,7 @@ namespace UE::NNEHlslShaders::Internal
 		DECLARE_GLOBAL_SHADER(TElementWiseVariadicCS);
 		SHADER_USE_PARAMETER_STRUCT(TElementWiseVariadicCS, FHlslShaderBase)
 
-		class FOperatorType : SHADER_PERMUTATION_ENUM_CLASS("OP_TYPENAME", NNE::Internal::EElementWiseVariadicOperatorType);
+		class FOperatorType : SHADER_PERMUTATION_ENUM_CLASS("OP_TYPENAME", EElementWiseVariadicOperatorType);
 		class FApplyScale : SHADER_PERMUTATION_BOOL("APPLYSCALE");
 		class FOutputAsInput : SHADER_PERMUTATION_BOOL("OUTPUTASINPUT");
 		class FNumInput : SHADER_PERMUTATION_RANGE_INT("NUMINPUT", 1, FElementWiseVariadicConstants::MAX_NUM_INPUT);
@@ -49,6 +49,6 @@ namespace UE::NNEHlslShaders::Internal
 
 	private:
 
-		static const FString GetOpFunc(NNE::Internal::EElementWiseVariadicOperatorType OpType);
+		static const FString GetOpFunc(EElementWiseVariadicOperatorType OpType);
 	};
 } // UE::NNEHlslShaders::Internal
\ No newline at end of file
diff --git a/Engine/Plugins/Experimental/NNERuntimeRDG/Source/NNEHlslShaders/Internal/NNEHlslShadersOperator.h b/Engine/Plugins/Experimental/NNERuntimeRDG/Source/NNEHlslShaders/Internal/NNEHlslShadersOperator.h
new file mode 100644
--- /dev/null
+++ b/Engine/Plugins/Experimental/NNERuntimeRDG/Source/NNEHlslShaders/Internal/NNEHlslShadersOperator.h
@@ -0,0 +1,89 @@
+// Copyright Epic Games, Inc. All Rights Reserved.
+
+#pragma once
+
+#include "HAL/Platform.h"
+
+namespace UE::NNEHlslShaders::Internal
+{
+
+	//One input element-wise operators
+	enum class EElementWiseUnaryOperatorType : uint8
+	{
+		Abs,
+		Acos,
+		Acosh,
+		Asin,
+		Asinh,
+		Atan,
+		Atanh,
+		//BitShift, //Note: need integer tensors
+		Ceil,
+		Clip,
+		Cos,
+		Cosh,
+		Elu,
+		Erf,
+		Exp,
+		Floor,
+		IsInf,
+		IsNan,
+		HardSigmoid,
+		HardSwish,
+		LeakyRelu,
+		Log,
+		Neg,
+		//Not,      //Note: need bool tensors
+		Reciprocal,
+		Relu,
+		Round,
+		Selu,
+		Sigmoid,
+		Sign,
+		Sin,
+		Sinh,
+		Softplus,
+		Softsign,
+		Sqrt,
+		Tan,
+		Tanh,
+
+		MAX
+	};
+
+	//Two inputs element-wise operators with multi-directional broadcast
+	//see [REDACTED_PROFILE_URL]/onnx/blob/main/docs/Broadcasting.md
+	enum class EElementWiseBinaryOperatorType : uint8
+	{
+		Add,
+		//And,           //Note: need boolean tensors
+		Div,
+		//Equal,         //Note: need boolean tensors
+		//Greater,       //Note: need boolean tensors
+		//GreaterOrEqual,//Note: need boolean tensors
+		//Less,          //Note: need boolean tensors
+		//LessOrEqual,   //Note: need boolean tensors
+		Mod,
+		Mul,
+		//Or,            //Note: need boolean tensors
+		Prelu,           //Note: only broadcast from slope to input.
+		Pow,
+		Sub,
+		//Xor,           //Note:  need boolean tensors
+
+		MAX
+	};
+
+	//Variable number of inputs element-wise operators with multi-directional broadcast
+	//see [REDACTED_PROFILE_URL]/onnx/blob/main/docs/Broadcasting.md
+	enum class EElementWiseVariadicOperatorType : uint8
+	{
+		Max,
+		Min,
+		Mean,
+		Sum,
+
+		MAX
+	};
+
+} // namespace UE::NNEHlslShaders::Internal
\ No newline at end of file
diff --git a/Engine/Plugins/Experimental/NNERuntimeRDG/Source/NNEHlslShaders/Private/NNEHlslShadersElementWiseBinaryCS.cpp b/Engine/Plugins/Experimental/NNERuntimeRDG/Source/NNEHlslShaders/Private/NNEHlslShadersElementWiseBinaryCS.cpp
--- a/Engine/Plugins/Experimental/NNERuntimeRDG/Source/NNEHlslShaders/Private/NNEHlslShadersElementWiseBinaryCS.cpp
+++ b/Engine/Plugins/Experimental/NNERuntimeRDG/Source/NNEHlslShaders/Private/NNEHlslShadersElementWiseBinaryCS.cpp
@@ -18,16 +18,16 @@ namespace UE::NNEHlslShaders::Internal
 		OutEnvironment.SetDefine(TEXT("ELEMENTWISE_OP(X,Y)"), *OpFunc);
 	}
 
-	const FString TElementWiseBinaryCS::GetOpFunc(NNE::Internal::EElementWiseBinaryOperatorType OpType)
+	const FString TElementWiseBinaryCS::GetOpFunc(EElementWiseBinaryOperatorType OpType)
 	{
-		FString OpTable[(int32) NNE::Internal::EElementWiseBinaryOperatorType::MAX];
+		FString OpTable[(int32) EElementWiseBinaryOperatorType::MAX];
 
-		for (int32 Idx = 0; Idx < (int32) NNE::Internal::EElementWiseBinaryOperatorType::MAX; ++Idx)
+		for (int32 Idx = 0; Idx < (int32) EElementWiseBinaryOperatorType::MAX; ++Idx)
 		{
 			OpTable[Idx] = FString("");
 		}
 
-#define OP(OpName, OpFunc) OpTable[(int32) NNE::Internal::EElementWiseBinaryOperatorType::OpName] = OpFunc
+#define OP(OpName, OpFunc) OpTable[(int32) EElementWiseBinaryOperatorType::OpName] = OpFunc
 		OP(Add,              TEXT("((X)+(Y))"));
 		//OP(And,            TEXT("((X)&&(Y))"));
 		OP(Div,              TEXT("((X)/(Y))"));
diff --git a/Engine/Plugins/Experimental/NNERuntimeRDG/Source/NNEHlslShaders/Private/NNEHlslShadersElementWiseUnaryCS.cpp b/Engine/Plugins/Experimental/NNERuntimeRDG/Source/NNEHlslShaders/Private/NNEHlslShadersElementWiseUnaryCS.cpp
--- a/Engine/Plugins/Experimental/NNERuntimeRDG/Source/NNEHlslShaders/Private/NNEHlslShadersElementWiseUnaryCS.cpp
+++ b/Engine/Plugins/Experimental/NNERuntimeRDG/Source/NNEHlslShaders/Private/NNEHlslShadersElementWiseUnaryCS.cpp
@@ -14,11 +14,11 @@ namespace UE::NNEHlslShaders::Internal
 
 		FPermutationDomain PermutationVector(Parameters.PermutationId);
 
-		NNE::Internal::EElementWiseUnaryOperatorType OpType = PermutationVector.Get<TElementWiseUnaryCS::FOperatorType>();
+		EElementWiseUnaryOperatorType OpType = PermutationVector.Get<TElementWiseUnaryCS::FOperatorType>();
 		const bool bAlphaOnGpu = PermutationVector.Get<TElementWiseUnaryCS::FAlphaOnGPU>();
 		const bool bBetaOnGpu = PermutationVector.Get<TElementWiseUnaryCS::FBetaOnGPU>();
 
-		return OpType == NNE::Internal::EElementWiseUnaryOperatorType::Clip || (!bAlphaOnGpu && !bBetaOnGpu);
+		return OpType == EElementWiseUnaryOperatorType::Clip || (!bAlphaOnGpu && !bBetaOnGpu);
 	}
 
 	void TElementWiseUnaryCS::ModifyCompilationEnvironment(const FGlobalShaderPermutationParameters& InParameters, FShaderCompilerEnvironment& OutEnvironment)
@@ -34,16 +34,16 @@ namespace UE::NNEHlslShaders::Internal
 		OutEnvironment.SetDefine(TEXT("ELEMENTWISE_OP(X)"), *OpFunc);
 	}
 
-	const FString TElementWiseUnaryCS::GetOpFunc(NNE::Internal::EElementWiseUnaryOperatorType OpType)
+	const FString TElementWiseUnaryCS::GetOpFunc(EElementWiseUnaryOperatorType OpType)
 	{
-		FString OpTable[(int32) NNE::Internal::EElementWiseUnaryOperatorType::MAX];
+		FString OpTable[(int32) EElementWiseUnaryOperatorType::MAX];
 
-		for (int32 Idx = 0; Idx < (int32) NNE::Internal::EElementWiseUnaryOperatorType::MAX; ++Idx)
+		for (int32 Idx = 0; Idx < (int32) EElementWiseUnaryOperatorType::MAX; ++Idx)
 		{
 			OpTable[Idx] = FString("");
 		}
 
-#define OP(OpName, OpFunc) OpTable[(int32) NNE::Internal::EElementWiseUnaryOperatorType::OpName] = OpFunc
+#define OP(OpName, OpFunc) OpTable[(int32) EElementWiseUnaryOperatorType::OpName] = OpFunc
 		OP(Abs,         TEXT("abs(X)"));
 		OP(Acos,        TEXT("acos(X)"));
 		OP(Acosh,       TEXT("acosh(X)"));
diff --git a/Engine/Plugins/Experimental/NNERuntimeRDG/Source/NNEHlslShaders/Private/NNEHlslShadersElementWiseVariadicCS.cpp b/Engine/Plugins/Experimental/NNERuntimeRDG/Source/NNEHlslShaders/Private/NNEHlslShadersElementWiseVariadicCS.cpp
--- a/Engine/Plugins/Experimental/NNERuntimeRDG/Source/NNEHlslShaders/Private/NNEHlslShadersElementWiseVariadicCS.cpp
+++ b/Engine/Plugins/Experimental/NNERuntimeRDG/Source/NNEHlslShaders/Private/NNEHlslShadersElementWiseVariadicCS.cpp
@@ -18,16 +18,16 @@ namespace UE::NNEHlslShaders::Internal
 		OutEnvironment.SetDefine(TEXT("ELEMENTWISE_OP(X,Y)"), *OpFunc);
 	}
 
-	const FString TElementWiseVariadicCS::GetOpFunc(NNE::Internal::EElementWiseVariadicOperatorType OpType)
+	const FString TElementWiseVariadicCS::GetOpFunc(EElementWiseVariadicOperatorType OpType)
 	{
-		FString OpTable[(int32) NNE::Internal::EElementWiseVariadicOperatorType::MAX];
+		FString OpTable[(int32) EElementWiseVariadicOperatorType::MAX];
 
-		for (int32 Idx = 0; Idx < (int32) NNE::Internal::EElementWiseVariadicOperatorType::MAX; ++Idx)
+		for (int32 Idx = 0; Idx < (int32) EElementWiseVariadicOperatorType::MAX; ++Idx)
 		{
 			OpTable[Idx] = FString("");
 		}
 
-#define OP(OpName, OpFunc) OpTable[(int32) NNE::Internal::EElementWiseVariadicOperatorType::OpName] = OpFunc
+#define OP(OpName, OpFunc) OpTable[(int32) EElementWiseVariadicOperatorType::OpName] = OpFunc
 		OP(Max, TEXT("max(X,Y)"));
 		OP(Min, TEXT("min(X,Y)"));
 		OP(Mean, TEXT("((X)+(Y))"));
diff --git a/Engine/Plugins/Experimental/NNERuntimeRDG/Source/NNERuntimeRDG/Internal/NNERuntimeRDGHelperElementWiseBinary.h b/Engine/Plugins/Experimental/NNERuntimeRDG/Source/NNERuntimeRDG/Internal/NNERuntimeRDGHelperElementWiseBinary.h
--- a/Engine/Plugins/Experimental/NNERuntimeRDG/Source/NNERuntimeRDG/Internal/NNERuntimeRDGHelperElementWiseBinary.h
+++ b/Engine/Plugins/Experimental/NNERuntimeRDG/Source/NNERuntimeRDG/Internal/NNERuntimeRDGHelperElementWiseBinary.h
@@ -3,12 +3,12 @@
 #pragma once 
 
 #include "CoreMinimal.h"
-#include "NNEOperator.h"
+#include "NNEHlslShadersOperator.h"
 
 namespace UE::NNE::Internal { class FTensor; }
 
 namespace UE::NNERuntimeRDG::Internal::CPUHelper::ElementWiseBinary
 {
-	void NNERUNTIMERDG_API Apply(NNE::Internal::EElementWiseBinaryOperatorType OpType, const NNE::Internal::FTensor& LHSTensor, const NNE::Internal::FTensor& RHSTensor, NNE::Internal::FTensor& OutputTensor);
+	void NNERUNTIMERDG_API Apply(UE::NNEHlslShaders::Internal::EElementWiseBinaryOperatorType OpType, const NNE::Internal::FTensor& LHSTensor, const NNE::Internal::FTensor& RHSTensor, NNE::Internal::FTensor& OutputTensor);
 	
 } // UE::NNERuntimeRDG::Internal::CPUHelper::ElementWiseBinary
\ No newline at end of file
diff --git a/Engine/Plugins/Experimental/NNERuntimeRDG/Source/NNERuntimeRDG/Internal/NNERuntimeRDGHelperElementWiseUnary.h b/Engine/Plugins/Experimental/NNERuntimeRDG/Source/NNERuntimeRDG/Internal/NNERuntimeRDGHelperElementWiseUnary.h
--- a/Engine/Plugins/Experimental/NNERuntimeRDG/Source/NNERuntimeRDG/Internal/NNERuntimeRDGHelperElementWiseUnary.h
+++ b/Engine/Plugins/Experimental/NNERuntimeRDG/Source/NNERuntimeRDG/Internal/NNERuntimeRDGHelperElementWiseUnary.h
@@ -3,11 +3,12 @@
 #pragma once 
 
 #include "CoreMinimal.h"
-#include "NNEOperator.h"
+#include "NNEHlslShadersOperator.h"
 
 namespace UE::NNE::Internal { class FTensor; }
 
 namespace UE::NNERuntimeRDG::Internal::CPUHelper::ElementWiseUnary
 {
-	void NNERUNTIMERDG_API Apply(NNE::Internal::EElementWiseUnaryOperatorType OpType, const NNE::Internal::FTensor& Tensor, float Alpha, float Beta, float Gamma, NNE::Internal::FTensor& OutputTensor);
+	void NNERUNTIMERDG_API Apply(UE::NNEHlslShaders::Internal::EElementWiseUnaryOperatorType OpType, const NNE::Internal::FTensor& Tensor, float Alpha, float Beta, float Gamma, NNE::Internal::FTensor& OutputTensor);
+
 } // UE::NNERuntimeRDG::Internal::CPUHelper::ElementWiseUnary
\ No newline at end of file
diff --git a/Engine/Plugins/Experimental/NNERuntimeRDG/Source/NNERuntimeRDG/Private/Helper/NNERuntimeRDGHelperElementWiseBinary.cpp b/Engine/Plugins/Experimental/NNERuntimeRDG/Source/NNERuntimeRDG/Private/Helper/NNERuntimeRDGHelperElementWiseBinary.cpp
--- a/Engine/Plugins/Experimental/NNERuntimeRDG/Source/NNERuntimeRDG/Private/Helper/NNERuntimeRDGHelperElementWiseBinary.cpp
+++ b/Engine/Plugins/Experimental/NNERuntimeRDG/Source/NNERuntimeRDG/Private/Helper/NNERuntimeRDGHelperElementWiseBinary.cpp
@@ -9,33 +9,35 @@
 
 namespace UE::NNERuntimeRDG::Internal::CPUHelper::ElementWiseBinary
 {
-	template<NNE::Internal::EElementWiseBinaryOperatorType OpType, typename TData> TData Apply(TData X, TData Y);
+	using EElementWiseBinaryOperatorType = UE::NNEHlslShaders::Internal::EElementWiseBinaryOperatorType;
 
-	template<> float Apply<NNE::Internal::EElementWiseBinaryOperatorType::Add, float>(float X, float Y) { return X + Y; }
-	template<> float Apply<NNE::Internal::EElementWiseBinaryOperatorType::Div, float>(float X, float Y) { return X / Y; }
-	template<> float Apply<NNE::Internal::EElementWiseBinaryOperatorType::Mod, float>(float X, float Y) { return FMath::Fmod(X, Y); }
-	template<> float Apply<NNE::Internal::EElementWiseBinaryOperatorType::Mul, float>(float X, float Y) { return X * Y; }
-	template<> float Apply<NNE::Internal::EElementWiseBinaryOperatorType::Prelu, float>(float X, float Y) { return (X < 0.0f) ? (Y * X) : X; }
-	template<> float Apply<NNE::Internal::EElementWiseBinaryOperatorType::Pow, float>(float X, float Y) { return FMath::Pow(X, Y); }
-	template<> float Apply<NNE::Internal::EElementWiseBinaryOperatorType::Sub, float>(float X, float Y) { return X - Y; }
+	template<EElementWiseBinaryOperatorType OpType, typename TData> TData Apply(TData X, TData Y);
 
-	template<> int32 Apply<NNE::Internal::EElementWiseBinaryOperatorType::Add, int32>(int32 X, int32 Y) { return X + Y; }
-	template<> int32 Apply<NNE::Internal::EElementWiseBinaryOperatorType::Div, int32>(int32 X, int32 Y) { return X / Y; }
-	template<> int32 Apply<NNE::Internal::EElementWiseBinaryOperatorType::Mod, int32>(int32 X, int32 Y) { return X % Y; }
-	template<> int32 Apply<NNE::Internal::EElementWiseBinaryOperatorType::Mul, int32>(int32 X, int32 Y) { return X * Y; }
-	template<> int32 Apply<NNE::Internal::EElementWiseBinaryOperatorType::Prelu, int32>(int32 X, int32 Y) { return (X < 0.0f) ? (Y * X) : X; }
-	template<> int32 Apply<NNE::Internal::EElementWiseBinaryOperatorType::Pow, int32>(int32 X, int32 Y) { return (int32)FMath::Pow((float)X, (float)Y); }
-	template<> int32 Apply<NNE::Internal::EElementWiseBinaryOperatorType::Sub, int32>(int32 X, int32 Y) { return X - Y; }
+	template<> float Apply<EElementWiseBinaryOperatorType::Add, float>(float X, float Y) { return X + Y; }
+	template<> float Apply<EElementWiseBinaryOperatorType::Div, float>(float X, float Y) { return X / Y; }
+	template<> float Apply<EElementWiseBinaryOperatorType::Mod, float>(float X, float Y) { return FMath::Fmod(X, Y); }
+	template<> float Apply<EElementWiseBinaryOperatorType::Mul, float>(float X, float Y) { return X * Y; }
+	template<> float Apply<EElementWiseBinaryOperatorType::Prelu, float>(float X, float Y) { return (X < 0.0f) ? (Y * X) : X; }
+	template<> float Apply<EElementWiseBinaryOperatorType::Pow, float>(float X, float Y) { return FMath::Pow(X, Y); }
+	template<> float Apply<EElementWiseBinaryOperatorType::Sub, float>(float X, float Y) { return X - Y; }
 
-	template<> int64 Apply<NNE::Internal::EElementWiseBinaryOperatorType::Add, int64>(int64 X, int64 Y) { return X + Y; }
-	template<> int64 Apply<NNE::Internal::EElementWiseBinaryOperatorType::Div, int64>(int64 X, int64 Y) { return X / Y; }
-	template<> int64 Apply<NNE::Internal::EElementWiseBinaryOperatorType::Mod, int64>(int64 X, int64 Y) { return X % Y; }
-	template<> int64 Apply<NNE::Internal::EElementWiseBinaryOperatorType::Mul, int64>(int64 X, int64 Y) { return X * Y; }
-	template<> int64 Apply<NNE::Internal::EElementWiseBinaryOperatorType::Prelu, int64>(int64 X, int64 Y) { return (X < 0.0f) ? (Y * X) : X; }
-	template<> int64 Apply<NNE::Internal::EElementWiseBinaryOperatorType::Pow, int64>(int64 X, int64 Y) { return (int64)FMath::Pow((float)X, (float)Y); }
-	template<> int64 Apply<NNE::Internal::EElementWiseBinaryOperatorType::Sub, int64>(int64 X, int64 Y) { return X - Y; }
+	template<> int32 Apply<EElementWiseBinaryOperatorType::Add, int32>(int32 X, int32 Y) { return X + Y; }
+	template<> int32 Apply<EElementWiseBinaryOperatorType::Div, int32>(int32 X, int32 Y) { return X / Y; }
+	template<> int32 Apply<EElementWiseBinaryOperatorType::Mod, int32>(int32 X, int32 Y) { return X % Y; }
+	template<> int32 Apply<EElementWiseBinaryOperatorType::Mul, int32>(int32 X, int32 Y) { return X * Y; }
+	template<> int32 Apply<EElementWiseBinaryOperatorType::Prelu, int32>(int32 X, int32 Y) { return (X < 0.0f) ? (Y * X) : X; }
+	template<> int32 Apply<EElementWiseBinaryOperatorType::Pow, int32>(int32 X, int32 Y) { return (int32)FMath::Pow((float)X, (float)Y); }
+	template<> int32 Apply<EElementWiseBinaryOperatorType::Sub, int32>(int32 X, int32 Y) { return X - Y; }
 
-	template<NNE::Internal::EElementWiseBinaryOperatorType OpType, typename TData> void Apply(const NNE::Internal::FTensor& LHSTensor, const NNE::Internal::FTensor& RHSTensor, NNE::Internal::FTensor& OutputTensor)
+	template<> int64 Apply<EElementWiseBinaryOperatorType::Add, int64>(int64 X, int64 Y) { return X + Y; }
+	template<> int64 Apply<EElementWiseBinaryOperatorType::Div, int64>(int64 X, int64 Y) { return X / Y; }
+	template<> int64 Apply<EElementWiseBinaryOperatorType::Mod, int64>(int64 X, int64 Y) { return X % Y; }
+	template<> int64 Apply<EElementWiseBinaryOperatorType::Mul, int64>(int64 X, int64 Y) { return X * Y; }
+	template<> int64 Apply<EElementWiseBinaryOperatorType::Prelu, int64>(int64 X, int64 Y) { return (X < 0.0f) ? (Y * X) : X; }
+	template<> int64 Apply<EElementWiseBinaryOperatorType::Pow, int64>(int64 X, int64 Y) { return (int64)FMath::Pow((float)X, (float)Y); }
+	template<> int64 Apply<EElementWiseBinaryOperatorType::Sub, int64>(int64 X, int64 Y) { return X - Y; }
+
+	template<EElementWiseBinaryOperatorType OpType, typename TData> void Apply(const NNE::Internal::FTensor& LHSTensor, const NNE::Internal::FTensor& RHSTensor, NNE::Internal::FTensor& OutputTensor)
 	{
 		//Heuristic to avoid unexpected performance hit. This helper being intended for shape related arithmetic only.
 		static constexpr int32 MaxItemInInputTensors = NNE::FTensorShape::MaxRank * 2;
@@ -65,37 +67,37 @@ namespace UE::NNERuntimeRDG::Internal::CPUHelper::ElementWiseBinary
 		}
 	}
 
-	template<typename TData> void ApplyResolvedDataType(NNE::Internal::EElementWiseBinaryOperatorType OpType, const NNE::Internal::FTensor& LHSTensor, const NNE::Internal::FTensor& RHSTensor, NNE::Internal::FTensor& OutputTensor)
+	template<typename TData> void ApplyResolvedDataType(EElementWiseBinaryOperatorType OpType, const NNE::Internal::FTensor& LHSTensor, const NNE::Internal::FTensor& RHSTensor, NNE::Internal::FTensor& OutputTensor)
 	{
 		switch (OpType)
 		{
-		case NNE::Internal::EElementWiseBinaryOperatorType::Add:
-			Apply<NNE::Internal::EElementWiseBinaryOperatorType::Add, TData>(LHSTensor, RHSTensor, OutputTensor);
+		case EElementWiseBinaryOperatorType::Add:
+			Apply<EElementWiseBinaryOperatorType::Add, TData>(LHSTensor, RHSTensor, OutputTensor);
 			break;
-		case NNE::Internal::EElementWiseBinaryOperatorType::Div:
-			Apply<NNE::Internal::EElementWiseBinaryOperatorType::Div, TData>(LHSTensor, RHSTensor, OutputTensor);
+		case EElementWiseBinaryOperatorType::Div:
+			Apply<EElementWiseBinaryOperatorType::Div, TData>(LHSTensor, RHSTensor, OutputTensor);
 			break;
-		case NNE::Internal::EElementWiseBinaryOperatorType::Mod:
-			Apply<NNE::Internal::EElementWiseBinaryOperatorType::Mod, TData>(LHSTensor, RHSTensor, OutputTensor);
+		case EElementWiseBinaryOperatorType::Mod:
+			Apply<EElementWiseBinaryOperatorType::Mod, TData>(LHSTensor, RHSTensor, OutputTensor);
 			break;
-		case NNE::Internal::EElementWiseBinaryOperatorType::Mul:
-			Apply<NNE::Internal::EElementWiseBinaryOperatorType::Mul, TData>(LHSTensor, RHSTensor, OutputTensor);
+		case EElementWiseBinaryOperatorType::Mul:
+			Apply<EElementWiseBinaryOperatorType::Mul, TData>(LHSTensor, RHSTensor, OutputTensor);
 			break;
-		case NNE::Internal::EElementWiseBinaryOperatorType::Prelu:
-			Apply<NNE::Internal::EElementWiseBinaryOperatorType::Prelu, TData>(LHSTensor, RHSTensor, OutputTensor);
+		case EElementWiseBinaryOperatorType::Prelu:
+			Apply<EElementWiseBinaryOperatorType::Prelu, TData>(LHSTensor, RHSTensor, OutputTensor);
 			break;
-		case NNE::Internal::EElementWiseBinaryOperatorType::Pow:
-			Apply<NNE::Internal::EElementWiseBinaryOperatorType::Pow, TData>(LHSTensor, RHSTensor, OutputTensor);
+		case EElementWiseBinaryOperatorType::Pow:
+			Apply<EElementWiseBinaryOperatorType::Pow, TData>(LHSTensor, RHSTensor, OutputTensor);
 			break;
-		case NNE::Internal::EElementWiseBinaryOperatorType::Sub:
-			Apply<NNE::Internal::EElementWiseBinaryOperatorType::Sub, TData>(LHSTensor, RHSTensor, OutputTensor);
+		case EElementWiseBinaryOperatorType::Sub:
+			Apply<EElementWiseBinaryOperatorType::Sub, TData>(LHSTensor, RHSTensor, OutputTensor);
 			break;
 		default:
 			break;
 		}
 	}
 
-	void Apply(NNE::Internal::EElementWiseBinaryOperatorType OpType, const NNE::Internal::FTensor& LHSTensor, const NNE::Internal::FTensor& RHSTensor, NNE::Internal::FTensor& OutputTensor)
+	void Apply(EElementWiseBinaryOperatorType OpType, const NNE::Internal::FTensor& LHSTensor, const NNE::Internal::FTensor& RHSTensor, NNE::Internal::FTensor& OutputTensor)
 	{
 		switch (OutputTensor.GetDataType())
 		{
diff --git a/Engine/Plugins/Experimental/NNERuntimeRDG/Source/NNERuntimeRDG/Private/Helper/NNERuntimeRDGHelperElementWiseUnary.cpp b/Engine/Plugins/Experimental/NNERuntimeRDG/Source/NNERuntimeRDG/Private/Helper/NNERuntimeRDGHelperElementWiseUnary.cpp
--- a/Engine/Plugins/Experimental/NNERuntimeRDG/Source/NNERuntimeRDG/Private/Helper/NNERuntimeRDGHelperElementWiseUnary.cpp
+++ b/Engine/Plugins/Experimental/NNERuntimeRDG/Source/NNERuntimeRDG/Private/Helper/NNERuntimeRDGHelperElementWiseUnary.cpp
@@ -8,13 +8,15 @@
 
 namespace UE::NNERuntimeRDG::Internal::CPUHelper::ElementWiseUnary
 {
-	template<NNE::Internal::EElementWiseUnaryOperatorType OpType> float Apply(float X, float Alpha, float Beta, float Gamma);
+	using EElementWiseUnaryOperatorType = UE::NNEHlslShaders::Internal::EElementWiseUnaryOperatorType;
+
+	template<EElementWiseUnaryOperatorType OpType> float Apply(float X, float Alpha, float Beta, float Gamma);
 	
-	template<> float Apply<NNE::Internal::EElementWiseUnaryOperatorType::Abs>(float X, float Alpha, float Beta, float Gamma) { return FMath::Abs(X); }
+	template<> float Apply<EElementWiseUnaryOperatorType::Abs>(float X, float Alpha, float Beta, float Gamma) { return FMath::Abs(X); }
 
-	template<> float Apply<NNE::Internal::EElementWiseUnaryOperatorType::Acos>(float X, float Alpha, float Beta, float Gamma) { return FMath::Acos(X); }
+	template<> float Apply<EElementWiseUnaryOperatorType::Acos>(float X, float Alpha, float Beta, float Gamma) { return FMath::Acos(X); }
 	
-	template<> float Apply<NNE::Internal::EElementWiseUnaryOperatorType::Acosh>(float X, float Alpha, float Beta, float Gamma) {
+	template<> float Apply<EElementWiseUnaryOperatorType::Acosh>(float X, float Alpha, float Beta, float Gamma) {
 		//https://mathworld.wolfram.com/InverseHyperbolicCosine.html
 		float FloatNan = FMath::Sqrt(-1.0f);
 		float yAboveOne = FMath::Loge(X + FMath::Sqrt(X + 1.0f) * FMath::Sqrt(X - 1.0f));
@@ -28,107 +30,107 @@ namespace UE::NNERuntimeRDG::Internal::CPUHelper::ElementWiseUnary
 		}
 	}
 	
-	template<> float Apply<NNE::Internal::EElementWiseUnaryOperatorType::Asin>(float X, float Alpha, float Beta, float Gamma) { return FMath::Asin(X); }
+	template<> float Apply<EElementWiseUnaryOperatorType::Asin>(float X, float Alpha, float Beta, float Gamma) { return FMath::Asin(X); }
 	
-	template<> float Apply<NNE::Internal::EElementWiseUnaryOperatorType::Asinh>(float X, float Alpha, float Beta, float Gamma) {
+	template<> float Apply<EElementWiseUnaryOperatorType::Asinh>(float X, float Alpha, float Beta, float Gamma) {
 		//https://mathworld.wolfram.com/InverseHyperbolicSine.html
 		return FMath::Loge(X + FMath::Sqrt(1 + (X * X)));
 	}
 	
-	template<> float Apply<NNE::Internal::EElementWiseUnaryOperatorType::Atan>(float X, float Alpha, float Beta, float Gamma) { return FMath::Atan(X); }
+	template<> float Apply<EElementWiseUnaryOperatorType::Atan>(float X, float Alpha, float Beta, float Gamma) { return FMath::Atan(X); }
 	
-	template<> float Apply<NNE::Internal::EElementWiseUnaryOperatorType::Atanh>(float X, float Alpha, float Beta, float Gamma) {
+	template<> float Apply<EElementWiseUnaryOperatorType::Atanh>(float X, float Alpha, float Beta, float Gamma) {
 		//https://mathworld.wolfram.com/InverseHyperbolicTangent.html
 		return 0.5f * (FMath::Loge(1 + X) - FMath::Loge(1 - X));
 	}
 
-	template<> float Apply<NNE::Internal::EElementWiseUnaryOperatorType::Ceil>(float X, float Alpha, float Beta, float Gamma) { return FMath::CeilToFloat(X); }
+	template<> float Apply<EElementWiseUnaryOperatorType::Ceil>(float X, float Alpha, float Beta, float Gamma) { return FMath::CeilToFloat(X); }
 
-	template<> float Apply<NNE::Internal::EElementWiseUnaryOperatorType::Clip>(float X, float Alpha, float Beta, float Gamma) { return FMath::Clamp(X, Alpha, Beta); }
+	template<> float Apply<EElementWiseUnaryOperatorType::Clip>(float X, float Alpha, float Beta, float Gamma) { return FMath::Clamp(X, Alpha, Beta); }
 	
-	template<> float Apply<NNE::Internal::EElementWiseUnaryOperatorType::Cos>(float X, float Alpha, float Beta, float Gamma) { return FMath::Cos(X); }
+	template<> float Apply<EElementWiseUnaryOperatorType::Cos>(float X, float Alpha, float Beta, float Gamma) { return FMath::Cos(X); }
 	
-	template<> float Apply<NNE::Internal::EElementWiseUnaryOperatorType::Cosh>(float X, float Alpha, float Beta, float Gamma) {
+	template<> float Apply<EElementWiseUnaryOperatorType::Cosh>(float X, float Alpha, float Beta, float Gamma) {
 		//https://mathworld.wolfram.com/HyperbolicCosine.html
 		return 0.5f * (FMath::Exp(X) + FMath::Exp(-X));
 	}
 	
-	template<> float Apply<NNE::Internal::EElementWiseUnaryOperatorType::Elu>(float X, float Alpha, float Beta, float Gamma) {
+	template<> float Apply<EElementWiseUnaryOperatorType::Elu>(float X, float Alpha, float Beta, float Gamma) {
 		//[REDACTED_PROFILE_URL]/onnx/blob/main/docs/Operators.md#elu
 		float yNeg = Alpha * (FMath::Exp(X) - 1.0f);
 		float yPosOrZero = X;
 		return (X >= 0.0f) ? yPosOrZero : yNeg;
 	}
 	
-	template<> float Apply<NNE::Internal::EElementWiseUnaryOperatorType::Exp>(float X, float Alpha, float Beta, float Gamma) { return FMath::Exp(X); }
+	template<> float Apply<EElementWiseUnaryOperatorType::Exp>(float X, float Alpha, float Beta, float Gamma) { return FMath::Exp(X); }
 
-	template<> float Apply<NNE::Internal::EElementWiseUnaryOperatorType::Floor>(float X, float Alpha, float Beta, float Gamma) { return FMath::Floor(X); }
+	template<> float Apply<EElementWiseUnaryOperatorType::Floor>(float X, float Alpha, float Beta, float Gamma) { return FMath::Floor(X); }
 
-	template<> float Apply<NNE::Internal::EElementWiseUnaryOperatorType::IsInf>(float X, float Alpha, float Beta, float Gamma) { return !FMath::IsFinite(X); }
+	template<> float Apply<EElementWiseUnaryOperatorType::IsInf>(float X, float Alpha, float Beta, float Gamma) { return !FMath::IsFinite(X); }
 
-	template<> float Apply<NNE::Internal::EElementWiseUnaryOperatorType::IsNan>(float X, float Alpha, float Beta, float Gamma) { return FMath::IsNaN(X); }
+	template<> float Apply<EElementWiseUnaryOperatorType::IsNan>(float X, float Alpha, float Beta, float Gamma) { return FMath::IsNaN(X); }
 
-	template<> float Apply<NNE::Internal::EElementWiseUnaryOperatorType::HardSigmoid>(float X, float Alpha, float Beta, float Gamma) {
+	template<> float Apply<EElementWiseUnaryOperatorType::HardSigmoid>(float X, float Alpha, float Beta, float Gamma) {
 		//[REDACTED_PROFILE_URL]/onnx/blob/main/docs/Operators.md#hardSigmoid
 		return FMath::Max(0.0f, FMath::Min(1.0f, Alpha * X + Beta));
 	}
 
-	template<> float Apply<NNE::Internal::EElementWiseUnaryOperatorType::HardSwish>(float X, float Alpha, float Beta, float Gamma) {
+	template<> float Apply<EElementWiseUnaryOperatorType::HardSwish>(float X, float Alpha, float Beta, float Gamma) {
 		//[REDACTED_PROFILE_URL]/onnx/blob/main/docs/Operators.md#hardSwish
-		return Apply<NNE::Internal::EElementWiseUnaryOperatorType::HardSigmoid>(X, 1.0f / 6.0f, 0.5f, Gamma);
+		return Apply<EElementWiseUnaryOperatorType::HardSigmoid>(X, 1.0f / 6.0f, 0.5f, Gamma);
 	}
 
-	template<> float Apply<NNE::Internal::EElementWiseUnaryOperatorType::LeakyRelu>(float X, float Alpha, float Beta, float Gamma) { return (X >= 0.0f) ? X : Alpha * X; }
+	template<> float Apply<EElementWiseUnaryOperatorType::LeakyRelu>(float X, float Alpha, float Beta, float Gamma) { return (X >= 0.0f) ? X : Alpha * X; }
 
-	template<> float Apply<NNE::Internal::EElementWiseUnaryOperatorType::Log>(float X, float Alpha, float Beta, float Gamma) { return FMath::Loge(X); }
+	template<> float Apply<EElementWiseUnaryOperatorType::Log>(float X, float Alpha, float Beta, float Gamma) { return FMath::Loge(X); }
 
-	template<> float Apply<NNE::Internal::EElementWiseUnaryOperatorType::Neg>(float X, float Alpha, float Beta, float Gamma) { return -X; }
+	template<> float Apply<EElementWiseUnaryOperatorType::Neg>(float X, float Alpha, float Beta, float Gamma) { return -X; }
 
-	template<> float Apply<NNE::Internal::EElementWiseUnaryOperatorType::Reciprocal>(float X, float Alpha, float Beta, float Gamma) { return 1.0f/X; }
+	template<> float Apply<EElementWiseUnaryOperatorType::Reciprocal>(float X, float Alpha, float Beta, float Gamma) { return 1.0f/X; }
 
-	template<> float Apply<NNE::Internal::EElementWiseUnaryOperatorType::Relu>(float X, float Alpha, float Beta, float Gamma) { return FMath::Max(X, 0.0f); }
+	template<> float Apply<EElementWiseUnaryOperatorType::Relu>(float X, float Alpha, float Beta, float Gamma) { return FMath::Max(X, 0.0f); }
 
-	template<> float Apply<NNE::Internal::EElementWiseUnaryOperatorType::Round>(float X, float Alpha, float Beta, float Gamma) { return FMath::RoundToFloat(X); }
+	template<> float Apply<EElementWiseUnaryOperatorType::Round>(float X, float Alpha, float Beta, float Gamma) { return FMath::RoundToFloat(X); }
 
-	template<> float Apply<NNE::Internal::EElementWiseUnaryOperatorType::Selu>(float X, float Alpha, float Beta, float Gamma) {
+	template<> float Apply<EElementWiseUnaryOperatorType::Selu>(float X, float Alpha, float Beta, float Gamma) {
 		//[REDACTED_PROFILE_URL]/onnx/blob/main/docs/Operators.md#Selu
 		float yNegOrZero = Gamma * (Alpha * FMath::Exp(X) - Alpha);
 		float yPos = Gamma * X;
 		return (X > 0.0f) ? yPos : yNegOrZero;
 	}
 	
-	template<> float Apply<NNE::Internal::EElementWiseUnaryOperatorType::Sigmoid>(float X, float Alpha, float Beta, float Gamma) { return 1.0f / (1.0f + FMath::Exp(-X)); }
+	template<> float Apply<EElementWiseUnaryOperatorType::Sigmoid>(float X, float Alpha, float Beta, float Gamma) { return 1.0f / (1.0f + FMath::Exp(-X)); }
 	
-	template<> float Apply<NNE::Internal::EElementWiseUnaryOperatorType::Sign>(float X, float Alpha, float Beta, float Gamma) { return FMath::Sign(X); }
+	template<> float Apply<EElementWiseUnaryOperatorType::Sign>(float X, float Alpha, float Beta, float Gamma) { return FMath::Sign(X); }
 	
-	template<> float Apply<NNE::Internal::EElementWiseUnaryOperatorType::Sin>(float X, float Alpha, float Beta, float Gamma) { return FMath::Sin(X); }
+	template<> float Apply<EElementWiseUnaryOperatorType::Sin>(float X, float Alpha, float Beta, float Gamma) { return FMath::Sin(X); }
 	
-	template<> float Apply<NNE::Internal::EElementWiseUnaryOperatorType::Sinh>(float X, float Alpha, float Beta, float Gamma) { return FMath::Sinh(X); }
+	template<> float Apply<EElementWiseUnaryOperatorType::Sinh>(float X, float Alpha, float Beta, float Gamma) { return FMath::Sinh(X); }
 	
-	template<> float Apply<NNE::Internal::EElementWiseUnaryOperatorType::Softplus>(float X, float Alpha, float Beta, float Gamma) { return FMath::Loge(FMath::Exp(X) + 1.0f); }
+	template<> float Apply<EElementWiseUnaryOperatorType::Softplus>(float X, float Alpha, float Beta, float Gamma) { return FMath::Loge(FMath::Exp(X) + 1.0f); }
 	
-	template<> float Apply<NNE::Internal::EElementWiseUnaryOperatorType::Softsign>(float X, float Alpha, float Beta, float Gamma) { return X / (1.0f + FMath::Abs(X)); }
+	template<> float Apply<EElementWiseUnaryOperatorType::Softsign>(float X, float Alpha, float Beta, float Gamma) { return X / (1.0f + FMath::Abs(X)); }
 	
-	template<> float Apply<NNE::Internal::EElementWiseUnaryOperatorType::Sqrt>(float X, float Alpha, float Beta, float Gamma) { return FMath::Sqrt(X); }
+	template<> float Apply<EElementWiseUnaryOperatorType::Sqrt>(float X, float Alpha, float Beta, float Gamma) { return FMath::Sqrt(X); }
 	
-	template<> float Apply<NNE::Internal::EElementWiseUnaryOperatorType::Tan>(float X, float Alpha, float Beta, float Gamma) { return FMath::Tan(X); }
+	template<> float Apply<EElementWiseUnaryOperatorType::Tan>(float X, float Alpha, float Beta, float Gamma) { return FMath::Tan(X); }
 
-	template<> float Apply<NNE::Internal::EElementWiseUnaryOperatorType::Tanh>(float X, float Alpha, float Beta, float Gamma) {
+	template<> float Apply<EElementWiseUnaryOperatorType::Tanh>(float X, float Alpha, float Beta, float Gamma) {
 		//https://mathworld.wolfram.com/HyperbolicTangent.html
-		float SinhValue = Apply<NNE::Internal::EElementWiseUnaryOperatorType::Sinh>(X, Alpha, Beta, Gamma);
-		float CoshValue = Apply<NNE::Internal::EElementWiseUnaryOperatorType::Cosh>(X, Alpha, Beta, Gamma);
+		float SinhValue = Apply<EElementWiseUnaryOperatorType::Sinh>(X, Alpha, Beta, Gamma);
+		float CoshValue = Apply<EElementWiseUnaryOperatorType::Cosh>(X, Alpha, Beta, Gamma);
 		return SinhValue / CoshValue;
 	}
 
-	template<> float Apply<NNE::Internal::EElementWiseUnaryOperatorType::Erf>(float X, float Alpha, float Beta, float Gamma) {
+	template<> float Apply<EElementWiseUnaryOperatorType::Erf>(float X, float Alpha, float Beta, float Gamma) {
 		//https://aapt.scitation.org/doi/abs/10.1119/1.15018?journalCode=ajp
 		float a = 167.0f / 148.0f;
 		float b = 11.0f / 109.0f;
 		float x3 = X * X * X;
-		return Apply<NNE::Internal::EElementWiseUnaryOperatorType::Tanh>(a * X + b * x3, Alpha, Beta, Gamma);
+		return Apply<EElementWiseUnaryOperatorType::Tanh>(a * X + b * x3, Alpha, Beta, Gamma);
 	}
 
-	template<NNE::Internal::EElementWiseUnaryOperatorType OpType> void Apply(const NNE::Internal::FTensor& Tensor, float Alpha, float Beta, float Gamma, NNE::Internal::FTensor& OutputTensor)
+	template<EElementWiseUnaryOperatorType OpType> void Apply(const NNE::Internal::FTensor& Tensor, float Alpha, float Beta, float Gamma, NNE::Internal::FTensor& OutputTensor)
 	{
 		//Heuristic to avoid unexpected performance hit. This helper being intended for shape related arithmetic only.
 		static constexpr int32 MaxItemInInputTensors = NNE::FTensorShape::MaxRank * 2;
@@ -165,114 +167,114 @@ namespace UE::NNERuntimeRDG::Internal::CPUHelper::ElementWiseUnary
 		}
 	}
 
-	void Apply(NNE::Internal::EElementWiseUnaryOperatorType OpType, const NNE::Internal::FTensor& Tensor, float Alpha, float Beta, float Gamma, NNE::Internal::FTensor& OutputTensor)
+	void Apply(EElementWiseUnaryOperatorType OpType, const NNE::Internal::FTensor& Tensor, float Alpha, float Beta, float Gamma, NNE::Internal::FTensor& OutputTensor)
 	{
 		switch (OpType)
 		{
-		case NNE::Internal::EElementWiseUnaryOperatorType::Abs:
-			Apply<NNE::Internal::EElementWiseUnaryOperatorType::Abs>(Tensor, Alpha, Beta, Gamma, OutputTensor);
+		case EElementWiseUnaryOperatorType::Abs:
+			Apply<EElementWiseUnaryOperatorType::Abs>(Tensor, Alpha, Beta, Gamma, OutputTensor);
 			break;
-		case NNE::Internal::EElementWiseUnaryOperatorType::Acos:
-			Apply<NNE::Internal::EElementWiseUnaryOperatorType::Acos>(Tensor, Alpha, Beta, Gamma, OutputTensor);
+		case EElementWiseUnaryOperatorType::Acos:
+			Apply<EElementWiseUnaryOperatorType::Acos>(Tensor, Alpha, Beta, Gamma, OutputTensor);
 			break;
-		case NNE::Internal::EElementWiseUnaryOperatorType::Acosh:
-			Apply<NNE::Internal::EElementWiseUnaryOperatorType::Acosh>(Tensor, Alpha, Beta, Gamma, OutputTensor);
+		case EElementWiseUnaryOperatorType::Acosh:
+			Apply<EElementWiseUnaryOperatorType::Acosh>(Tensor, Alpha, Beta, Gamma, OutputTensor);
 			break;
-		case NNE::Internal::EElementWiseUnaryOperatorType::Asin:
-			Apply<NNE::Internal::EElementWiseUnaryOperatorType::Asin>(Tensor, Alpha, Beta, Gamma, OutputTensor);
+		case EElementWiseUnaryOperatorType::Asin:
+			Apply<EElementWiseUnaryOperatorType::Asin>(Tensor, Alpha, Beta, Gamma, OutputTensor);
 			break;
-		case NNE::Internal::EElementWiseUnaryOperatorType::Asinh:
-			Apply<NNE::Internal::EElementWiseUnaryOperatorType::Asinh>(Tensor, Alpha, Beta, Gamma, OutputTensor);
+		case EElementWiseUnaryOperatorType::Asinh:
+			Apply<EElementWiseUnaryOperatorType::Asinh>(Tensor, Alpha, Beta, Gamma, OutputTensor);
 			break;
-		case NNE::Internal::EElementWiseUnaryOperatorType::Atan:
-			Apply<NNE::Internal::EElementWiseUnaryOperatorType::Atan>(Tensor, Alpha, Beta, Gamma, OutputTensor);
+		case EElementWiseUnaryOperatorType::Atan:
+			Apply<EElementWiseUnaryOperatorType::Atan>(Tensor, Alpha, Beta, Gamma, OutputTensor);
 			break;
-		case NNE::Internal::EElementWiseUnaryOperatorType::Atanh:
-			Apply<NNE::Internal::EElementWiseUnaryOperatorType::Atanh>(Tensor, Alpha, Beta, Gamma, OutputTensor);
+		case EElementWiseUnaryOperatorType::Atanh:
+			Apply<EElementWiseUnaryOperatorType::Atanh>(Tensor, Alpha, Beta, Gamma, OutputTensor);
 			break;
-		case NNE::Internal::EElementWiseUnaryOperatorType::Ceil:
-			Apply<NNE::Internal::EElementWiseUnaryOperatorType::Ceil>(Tensor, Alpha, Beta, Gamma, OutputTensor);
+		case EElementWiseUnaryOperatorType::Ceil:
+			Apply<EElementWiseUnaryOperatorType::Ceil>(Tensor, Alpha, Beta, Gamma, OutputTensor);
 			break;
-		case NNE::Internal::EElementWiseUnaryOperatorType::Clip:
-			Apply<NNE::Internal::EElementWiseUnaryOperatorType::Clip>(Tensor, Alpha, Beta, Gamma, OutputTensor);
+		case EElementWiseUnaryOperatorType::Clip:
+			Apply<EElementWiseUnaryOperatorType::Clip>(Tensor, Alpha, Beta, Gamma, OutputTensor);
 			break;
-		case NNE::Internal::EElementWiseUnaryOperatorType::Cos:
-			Apply<NNE::Internal::EElementWiseUnaryOperatorType::Cos>(Tensor, Alpha, Beta, Gamma, OutputTensor);
+		case EElementWiseUnaryOperatorType::Cos:
+			Apply<EElementWiseUnaryOperatorType::Cos>(Tensor, Alpha, Beta, Gamma, OutputTensor);
 			break;
-		case NNE::Internal::EElementWiseUnaryOperatorType::Cosh:
-			Apply<NNE::Internal::EElementWiseUnaryOperatorType::Cosh>(Tensor, Alpha, Beta, Gamma, OutputTensor);
+		case EElementWiseUnaryOperatorType::Cosh:
+			Apply<EElementWiseUnaryOperatorType::Cosh>(Tensor, Alpha, Beta, Gamma, OutputTensor);
 			break;
-		case NNE::Internal::EElementWiseUnaryOperatorType::Elu:
-			Apply<NNE::Internal::EElementWiseUnaryOperatorType::Elu>(Tensor, Alpha, Beta, Gamma, OutputTensor);
+		case EElementWiseUnaryOperatorType::Elu:
+			Apply<EElementWiseUnaryOperatorType::Elu>(Tensor, Alpha, Beta, Gamma, OutputTensor);
 			break;
-		case NNE::Internal::EElementWiseUnaryOperatorType::Erf:
-			Apply<NNE::Internal::EElementWiseUnaryOperatorType::Erf>(Tensor, Alpha, Beta, Gamma, OutputTensor);
+		case EElementWiseUnaryOperatorType::Erf:
+			Apply<EElementWiseUnaryOperatorType::Erf>(Tensor, Alpha, Beta, Gamma, OutputTensor);
 			break;
-		case NNE::Internal::EElementWiseUnaryOperatorType::Exp:
-			Apply<NNE::Internal::EElementWiseUnaryOperatorType::Exp>(Tensor, Alpha, Beta, Gamma, OutputTensor);
+		case EElementWiseUnaryOperatorType::Exp:
+			Apply<EElementWiseUnaryOperatorType::Exp>(Tensor, Alpha, Beta, Gamma, OutputTensor);
 			break;
-		case NNE::Internal::EElementWiseUnaryOperatorType::Floor:
-			Apply<NNE::Internal::EElementWiseUnaryOperatorType::Floor>(Tensor, Alpha, Beta, Gamma, OutputTensor);
+		case EElementWiseUnaryOperatorType::Floor:
+			Apply<EElementWiseUnaryOperatorType::Floor>(Tensor, Alpha, Beta, Gamma, OutputTensor);
 			break;
-		case NNE::Internal::EElementWiseUnaryOperatorType::IsInf:
-			Apply<NNE::Internal::EElementWiseUnaryOperatorType::IsInf>(Tensor, Alpha, Beta, Gamma, OutputTensor);
+		case EElementWiseUnaryOperatorType::IsInf:
+			Apply<EElementWiseUnaryOperatorType::IsInf>(Tensor, Alpha, Beta, Gamma, OutputTensor);
 			break;
-		case NNE::Internal::EElementWiseUnaryOperatorType::IsNan:
-			Apply<NNE::Internal::EElementWiseUnaryOperatorType::IsNan>(Tensor, Alpha, Beta, Gamma, OutputTensor);
+		case EElementWiseUnaryOperatorType::IsNan:
+			Apply<EElementWiseUnaryOperatorType::IsNan>(Tensor, Alpha, Beta, Gamma, OutputTensor);
 			break;
-		case NNE::Internal::EElementWiseUnaryOperatorType::HardSigmoid:
-			Apply<NNE::Internal::EElementWiseUnaryOperatorType::HardSigmoid>(Tensor, Alpha, Beta, Gamma, OutputTensor);
+		case EElementWiseUnaryOperatorType::HardSigmoid:
+			Apply<EElementWiseUnaryOperatorType::HardSigmoid>(Tensor, Alpha, Beta, Gamma, OutputTensor);
 			break;
-		case NNE::Internal::EElementWiseUnaryOperatorType::HardSwish:
-			Apply<NNE::Internal::EElementWiseUnaryOperatorType::HardSwish>(Tensor, Alpha, Beta, Gamma, OutputTensor);
+		case EElementWiseUnaryOperatorType::HardSwish:
+			Apply<EElementWiseUnaryOperatorType::HardSwish>(Tensor, Alpha, Beta, Gamma, OutputTensor);
 			break;
-		case NNE::Internal::EElementWiseUnaryOperatorType::LeakyRelu:
-			Apply<NNE::Internal::EElementWiseUnaryOperatorType::LeakyRelu>(Tensor, Alpha, Beta, Gamma, OutputTensor);
+		case EElementWiseUnaryOperatorType::LeakyRelu:
+			Apply<EElementWiseUnaryOperatorType::LeakyRelu>(Tensor, Alpha, Beta, Gamma, OutputTensor);
 			break;
-		case NNE::Internal::EElementWiseUnaryOperatorType::Log:
-			Apply<NNE::Internal::EElementWiseUnaryOperatorType::Log>(Tensor, Alpha, Beta, Gamma, OutputTensor);
+		case EElementWiseUnaryOperatorType::Log:
+			Apply<EElementWiseUnaryOperatorType::Log>(Tensor, Alpha, Beta, Gamma, OutputTensor);
 			break;
-		case NNE::Internal::EElementWiseUnaryOperatorType::Neg:
-			Apply<NNE::Internal::EElementWiseUnaryOperatorType::Neg>(Tensor, Alpha, Beta, Gamma, OutputTensor);
+		case EElementWiseUnaryOperatorType::Neg:
+			Apply<EElementWiseUnaryOperatorType::Neg>(Tensor, Alpha, Beta, Gamma, OutputTensor);
 			break;
-		case NNE::Internal::EElementWiseUnaryOperatorType::Reciprocal:
-			Apply<NNE::Internal::EElementWiseUnaryOperatorType::Reciprocal>(Tensor, Alpha, Beta, Gamma, OutputTensor);
+		case EElementWiseUnaryOperatorType::Reciprocal:
+			Apply<EElementWiseUnaryOperatorType::Reciprocal>(Tensor, Alpha, Beta, Gamma, OutputTensor);
 			break;
-		case NNE::Internal::EElementWiseUnaryOperatorType::Relu:
-			Apply<NNE::Internal::EElementWiseUnaryOperatorType::Relu>(Tensor, Alpha, Beta, Gamma, OutputTensor);
+		case EElementWiseUnaryOperatorType::Relu:
+			Apply<EElementWiseUnaryOperatorType::Relu>(Tensor, Alpha, Beta, Gamma, OutputTensor);
 			break;
-		case NNE::Internal::EElementWiseUnaryOperatorType::Round:
-			Apply<NNE::Internal::EElementWiseUnaryOperatorType::Round>(Tensor, Alpha, Beta, Gamma, OutputTensor);
+		case EElementWiseUnaryOperatorType::Round:
+			Apply<EElementWiseUnaryOperatorType::Round>(Tensor, Alpha, Beta, Gamma, OutputTensor);
 			break;
-		case NNE::Internal::EElementWiseUnaryOperatorType::Selu:
-			Apply<NNE::Internal::EElementWiseUnaryOperatorType::Selu>(Tensor, Alpha, Beta, Gamma, OutputTensor);
+		case EElementWiseUnaryOperatorType::Selu:
+			Apply<EElementWiseUnaryOperatorType::Selu>(Tensor, Alpha, Beta, Gamma, OutputTensor);
 			break;
-		case NNE::Internal::EElementWiseUnaryOperatorType::Sigmoid:
-			Apply<NNE::Internal::EElementWiseUnaryOperatorType::Sigmoid>(Tensor, Alpha, Beta, Gamma, OutputTensor);
+		case EElementWiseUnaryOperatorType::Sigmoid:
+			Apply<EElementWiseUnaryOperatorType::Sigmoid>(Tensor, Alpha, Beta, Gamma, OutputTensor);
 			break;
-		case NNE::Internal::EElementWiseUnaryOperatorType::Sign:
-			Apply<NNE::Internal::EElementWiseUnaryOperatorType::Sign>(Tensor, Alpha, Beta, Gamma, OutputTensor);
+		case EElementWiseUnaryOperatorType::Sign:
+			Apply<EElementWiseUnaryOperatorType::Sign>(Tensor, Alpha, Beta, Gamma, OutputTensor);
 			break;
-		case NNE::Internal::EElementWiseUnaryOperatorType::Sin:
-			Apply<NNE::Internal::EElementWiseUnaryOperatorType::Sin>(Tensor, Alpha, Beta, Gamma, OutputTensor);
+		case EElementWiseUnaryOperatorType::Sin:
+			Apply<EElementWiseUnaryOperatorType::Sin>(Tensor, Alpha, Beta, Gamma, OutputTensor);
 			break;
-		case NNE::Internal::EElementWiseUnaryOperatorType::Sinh:
-			Apply<NNE::Internal::EElementWiseUnaryOperatorType::Sinh>(Tensor, Alpha, Beta, Gamma, OutputTensor);
+		case EElementWiseUnaryOperatorType::Sinh:
+			Apply<EElementWiseUnaryOperatorType::Sinh>(Tensor, Alpha, Beta, Gamma, OutputTensor);
 			break;
-		case NNE::Internal::EElementWiseUnaryOperatorType::Softplus:
-			Apply<NNE::Internal::EElementWiseUnaryOperatorType::Softplus>(Tensor, Alpha, Beta, Gamma, OutputTensor);
+		case EElementWiseUnaryOperatorType::Softplus:
+			Apply<EElementWiseUnaryOperatorType::Softplus>(Tensor, Alpha, Beta, Gamma, OutputTensor);
 			break;
-		case NNE::Internal::EElementWiseUnaryOperatorType::Softsign:
-			Apply<NNE::Internal::EElementWiseUnaryOperatorType::Softsign>(Tensor, Alpha, Beta, Gamma, OutputTensor);
+		case EElementWiseUnaryOperatorType::Softsign:
+			Apply<EElementWiseUnaryOperatorType::Softsign>(Tensor, Alpha, Beta, Gamma, OutputTensor);
 			break;
-		case NNE::Internal::EElementWiseUnaryOperatorType::Sqrt:
-			Apply<NNE::Internal::EElementWiseUnaryOperatorType::Sqrt>(Tensor, Alpha, Beta, Gamma, OutputTensor);
+		case EElementWiseUnaryOperatorType::Sqrt:
+			Apply<EElementWiseUnaryOperatorType::Sqrt>(Tensor, Alpha, Beta, Gamma, OutputTensor);
 			break;
-		case NNE::Internal::EElementWiseUnaryOperatorType::Tan:
-			Apply<NNE::Internal::EElementWiseUnaryOperatorType::Tan>(Tensor, Alpha, Beta, Gamma, OutputTensor);
+		case EElementWiseUnaryOperatorType::Tan:
+			Apply<EElementWiseUnaryOperatorType::Tan>(Tensor, Alpha, Beta, Gamma, OutputTensor);
 			break;
-		case NNE::Internal::EElementWiseUnaryOperatorType::Tanh:
-			Apply<NNE::Internal::EElementWiseUnaryOperatorType::Tanh>(Tensor, Alpha, Beta, Gamma, OutputTensor);
+		case EElementWiseUnaryOperatorType::Tanh:
+			Apply<EElementWiseUnaryOperatorType::Tanh>(Tensor, Alpha, Beta, Gamma, OutputTensor);
 			break;
 		default:
 			break;
diff --git a/Engine/Plugins/Experimental/NNERuntimeRDG/Source/NNERuntimeRDG/Private/Hlsl/NNERuntimeRDGElementWiseBinary.cpp b/Engine/Plugins/Experimental/NNERuntimeRDG/Source/NNERuntimeRDG/Private/Hlsl/NNERuntimeRDGElementWiseBinary.cpp
--- a/Engine/Plugins/Experimental/NNERuntimeRDG/Source/NNERuntimeRDG/Private/Hlsl/NNERuntimeRDGElementWiseBinary.cpp
+++ b/Engine/Plugins/Experimental/NNERuntimeRDG/Source/NNERuntimeRDG/Private/Hlsl/NNERuntimeRDGElementWiseBinary.cpp
@@ -15,11 +15,12 @@ namespace UE::NNERuntimeRDG::Private::Hlsl
 
 	using TElementWiseBinaryCS = typename UE::NNEHlslShaders::Internal::TElementWiseBinaryCS;
 	using FElementWiseBinaryConstants = UE::NNEHlslShaders::Internal::FElementWiseBinaryConstants;
+	using EElementWiseBinaryOperatorType = UE::NNEHlslShaders::Internal::EElementWiseBinaryOperatorType;
 
 	/**
 	 * Binary element-wise operator implementation
 	 */
-	template<NNE::Internal::EElementWiseBinaryOperatorType OpType>
+	template<EElementWiseBinaryOperatorType OpType>
 	class TElementWiseBinary : public FOperatorHlsl
 	{
 	public:
@@ -147,7 +148,7 @@ namespace UE::NNERuntimeRDG::Private::Hlsl
 		return bIsValid;
 	}
 
-	template<NNE::Internal::EElementWiseBinaryOperatorType OpType>
+	template<EElementWiseBinaryOperatorType OpType>
 	FOperatorHlsl* CreateElementWiseBinaryOperator()
 	{
 		return new TElementWiseBinary<OpType>();
@@ -156,7 +157,7 @@ namespace UE::NNERuntimeRDG::Private::Hlsl
 	bool RegisterElementWiseBinaryOperators(FOperatorRegistryHlsl& Registry)
 	{
 		// Note: support of a particular version is partial with respect to tensor data types (only the most typical ones are usually supported).
-#define OP(Name, Version) Registry.OpAdd({{TEXT(#Name), TEXT("Onnx")}, Version}, CreateElementWiseBinaryOperator<NNE::Internal::EElementWiseBinaryOperatorType::Name>, ValidateElementWiseBinaryOperator);
+#define OP(Name, Version) Registry.OpAdd({{TEXT(#Name), TEXT("Onnx")}, Version}, CreateElementWiseBinaryOperator<EElementWiseBinaryOperatorType::Name>, ValidateElementWiseBinaryOperator);
 #define OP_ALL_VERSIONS(Name) \
 OP(Name, 6) \
 OP(Name, 7) \
diff --git a/Engine/Plugins/Experimental/NNERuntimeRDG/Source/NNERuntimeRDG/Private/Hlsl/NNERuntimeRDGElementWiseUnary.cpp b/Engine/Plugins/Experimental/NNERuntimeRDG/Source/NNERuntimeRDG/Private/Hlsl/NNERuntimeRDGElementWiseUnary.cpp
--- a/Engine/Plugins/Experimental/NNERuntimeRDG/Source/NNERuntimeRDG/Private/Hlsl/NNERuntimeRDGElementWiseUnary.cpp
+++ b/Engine/Plugins/Experimental/NNERuntimeRDG/Source/NNERuntimeRDG/Private/Hlsl/NNERuntimeRDGElementWiseUnary.cpp
@@ -16,11 +16,12 @@ namespace UE::NNERuntimeRDG::Private::Hlsl
 
 	using TElementWiseUnaryCS = typename UE::NNEHlslShaders::Internal::TElementWiseUnaryCS;
 	using FElementWiseUnaryConstants = UE::NNEHlslShaders::Internal::FElementWiseUnaryConstants;
+	using EElementWiseUnaryOperatorType = UE::NNEHlslShaders::Internal::EElementWiseUnaryOperatorType;
 
 	/**
 	 * Unary element-wise operator implementation
 	 */
-	template<NNE::Internal::EElementWiseUnaryOperatorType OpType>
+	template<EElementWiseUnaryOperatorType OpType>
 	class TElementWiseUnary : public FOperatorHlsl
 	{
 	public:
@@ -40,14 +41,14 @@ namespace UE::NNERuntimeRDG::Private::Hlsl
 
 		virtual int PrepareOutputs(TConstArrayView<NNE::Internal::FTensorRef> InputTensors, TArrayView<NNE::Internal::FTensorRef> OutputTensors) override
 		{
-			check((OpType == NNE::Internal::EElementWiseUnaryOperatorType::Clip && Version >= 11) ? (InputTensors.Num() >= 1 && InputTensors.Num() <= 3) : InputTensors.Num() == 1);
+			check((OpType == EElementWiseUnaryOperatorType::Clip && Version >= 11) ? (InputTensors.Num() >= 1 && InputTensors.Num() <= 3) : InputTensors.Num() == 1);
 			check(OutputTensors.Num() == 1);
 			
 			OutputTensors[0]->SetShape(InputTensors[0]->GetShape());
 
 			const NNE::Internal::FTensor& X = *InputTensors[0];
 
-			if(OpType == NNE::Internal::EElementWiseUnaryOperatorType::Clip && Version >= 11 && InputTensors.Num() >= 2)
+			if(OpType == EElementWiseUnaryOperatorType::Clip && Version >= 11 && InputTensors.Num() >= 2)
 			{
 				{
 					const NNE::Internal::FTensor& MinTensor = *InputTensors[1];
@@ -111,7 +112,7 @@ namespace UE::NNERuntimeRDG::Private::Hlsl
 			PermutationVector.Set<TElementWiseUnaryCS::FAlphaOnGPU>(false);
 			PermutationVector.Set<TElementWiseUnaryCS::FBetaOnGPU>(false);
 
-			if(OpType == NNE::Internal::EElementWiseUnaryOperatorType::Clip && Version >= 11 && InputTensors.Num() >= 2)
+			if(OpType == EElementWiseUnaryOperatorType::Clip && Version >= 11 && InputTensors.Num() >= 2)
 			{
 				{
 					const NNE::Internal::FTensor& MinTensor = *InputTensors[1];
@@ -150,27 +151,27 @@ namespace UE::NNERuntimeRDG::Private::Hlsl
 		}
 	};
 
-	template<> TElementWiseUnary<NNE::Internal::EElementWiseUnaryOperatorType::Selu>::TElementWiseUnary(int InVersion)
+	template<> TElementWiseUnary<EElementWiseUnaryOperatorType::Selu>::TElementWiseUnary(int InVersion)
 		: Alpha(1.67326319217681884765625f), Beta(0.0f), Gamma(1.05070102214813232421875f), Version(InVersion)
 	{
 	}
 
-	template<> TElementWiseUnary<NNE::Internal::EElementWiseUnaryOperatorType::Elu>::TElementWiseUnary(int InVersion)
+	template<> TElementWiseUnary<EElementWiseUnaryOperatorType::Elu>::TElementWiseUnary(int InVersion)
 		: Alpha(1.0f), Beta(0.0f), Gamma(0.0f), Version(InVersion)
 	{
 	}
 
-	template<> TElementWiseUnary<NNE::Internal::EElementWiseUnaryOperatorType::HardSigmoid>::TElementWiseUnary(int InVersion)
+	template<> TElementWiseUnary<EElementWiseUnaryOperatorType::HardSigmoid>::TElementWiseUnary(int InVersion)
 		: Alpha(0.2f), Beta(0.5f), Gamma(0.0f), Version(InVersion)
 	{
 	}
 
-	template<> TElementWiseUnary<NNE::Internal::EElementWiseUnaryOperatorType::LeakyRelu>::TElementWiseUnary(int InVersion)
+	template<> TElementWiseUnary<EElementWiseUnaryOperatorType::LeakyRelu>::TElementWiseUnary(int InVersion)
 		: Alpha(0.01f), Beta(0.0f), Gamma(0.0f), Version(InVersion)
 	{
 	}
 
-	template<> bool TElementWiseUnary<NNE::Internal::EElementWiseUnaryOperatorType::Clip>::Initialize(TConstArrayView<NNE::FTensorDesc> InputTensorDescs, TConstArrayView<NNE::FTensorDesc> OutputTensorDescs, const NNE::FAttributeMap& Attributes)
+	template<> bool TElementWiseUnary<EElementWiseUnaryOperatorType::Clip>::Initialize(TConstArrayView<NNE::FTensorDesc> InputTensorDescs, TConstArrayView<NNE::FTensorDesc> OutputTensorDescs, const NNE::FAttributeMap& Attributes)
 	{
 		check(Version >= 11 ? (InputTensorDescs.Num() >= 1 && InputTensorDescs.Num() <= 3) : InputTensorDescs.Num() == 1);
 		check(OutputTensorDescs.Num() == 1);
@@ -180,13 +181,13 @@ namespace UE::NNERuntimeRDG::Private::Hlsl
 		return true;
 	}
 
-	template<NNE::Internal::EElementWiseUnaryOperatorType OpType, int Version>
+	template<EElementWiseUnaryOperatorType OpType, int Version>
 	FOperatorHlsl* CreateElementWiseUnaryOperator()
 	{
 		return new TElementWiseUnary<OpType>(Version);
 	}
 	
-	template<NNE::Internal::EElementWiseUnaryOperatorType OpType>
+	template<EElementWiseUnaryOperatorType OpType>
 	bool ValidateElementWiseUnaryOperatorImpl(int Version, const NNE::FAttributeMap& AttributeMap, TConstArrayView<ENNETensorDataType> InputTypes, TConstArrayView<NNE::FSymbolicTensorShape> InputShapes)
 	{
 		bool bIsValid = true;
@@ -204,7 +205,7 @@ namespace UE::NNERuntimeRDG::Private::Hlsl
 	}
 
 	template<>
-	bool ValidateElementWiseUnaryOperatorImpl<NNE::Internal::EElementWiseUnaryOperatorType::Selu>(int Version, const NNE::FAttributeMap& AttributeMap, TConstArrayView<ENNETensorDataType> InputTypes, TConstArrayView<NNE::FSymbolicTensorShape> InputShapes)
+	bool ValidateElementWiseUnaryOperatorImpl<EElementWiseUnaryOperatorType::Selu>(int Version, const NNE::FAttributeMap& AttributeMap, TConstArrayView<ENNETensorDataType> InputTypes, TConstArrayView<NNE::FSymbolicTensorShape> InputShapes)
 	{
 		bool bIsValid = true;
 
@@ -223,7 +224,7 @@ namespace UE::NNERuntimeRDG::Private::Hlsl
 	}
 
 	template<>
-	bool ValidateElementWiseUnaryOperatorImpl<NNE::Internal::EElementWiseUnaryOperatorType::Elu>(int Version, const NNE::FAttributeMap& AttributeMap, TConstArrayView<ENNETensorDataType> InputTypes, TConstArrayView<NNE::FSymbolicTensorShape> InputShapes)
+	bool ValidateElementWiseUnaryOperatorImpl<EElementWiseUnaryOperatorType::Elu>(int Version, const NNE::FAttributeMap& AttributeMap, TConstArrayView<ENNETensorDataType> InputTypes, TConstArrayView<NNE::FSymbolicTensorShape> InputShapes)
 	{
 		bool bIsValid = true;
 
@@ -241,7 +242,7 @@ namespace UE::NNERuntimeRDG::Private::Hlsl
 	}
 
 	template<>
-	bool ValidateElementWiseUnaryOperatorImpl<NNE::Internal::EElementWiseUnaryOperatorType::HardSigmoid>(int Version, const NNE::FAttributeMap& AttributeMap, TConstArrayView<ENNETensorDataType> InputTypes, TConstArrayView<NNE::FSymbolicTensorShape> InputShapes)
+	bool ValidateElementWiseUnaryOperatorImpl<EElementWiseUnaryOperatorType::HardSigmoid>(int Version, const NNE::FAttributeMap& AttributeMap, TConstArrayView<ENNETensorDataType> InputTypes, TConstArrayView<NNE::FSymbolicTensorShape> InputShapes)
 	{
 		bool bIsValid = true;
 
@@ -260,7 +261,7 @@ namespace UE::NNERuntimeRDG::Private::Hlsl
 	}
 
 	template<>
-	bool ValidateElementWiseUnaryOperatorImpl<NNE::Internal::EElementWiseUnaryOperatorType::LeakyRelu>(int Version, const NNE::FAttributeMap& AttributeMap, TConstArrayView<ENNETensorDataType> InputTypes, TConstArrayView<NNE::FSymbolicTensorShape> InputShapes)
+	bool ValidateElementWiseUnaryOperatorImpl<EElementWiseUnaryOperatorType::LeakyRelu>(int Version, const NNE::FAttributeMap& AttributeMap, TConstArrayView<ENNETensorDataType> InputTypes, TConstArrayView<NNE::FSymbolicTensorShape> InputShapes)
 	{
 		bool bIsValid = true;
 
@@ -278,7 +279,7 @@ namespace UE::NNERuntimeRDG::Private::Hlsl
 	}
 
 	template<>
-	bool ValidateElementWiseUnaryOperatorImpl<NNE::Internal::EElementWiseUnaryOperatorType::Clip>(int Version, const NNE::FAttributeMap& AttributeMap, TConstArrayView<ENNETensorDataType> InputTypes, TConstArrayView<NNE::FSymbolicTensorShape> InputShapes)
+	bool ValidateElementWiseUnaryOperatorImpl<EElementWiseUnaryOperatorType::Clip>(int Version, const NNE::FAttributeMap& AttributeMap, TConstArrayView<ENNETensorDataType> InputTypes, TConstArrayView<NNE::FSymbolicTensorShape> InputShapes)
 	{
 		bool bIsValid = true;
 
@@ -304,7 +305,7 @@ namespace UE::NNERuntimeRDG::Private::Hlsl
 		return bIsValid;
 	}
 
-	template<NNE::Internal::EElementWiseUnaryOperatorType OpType, int Version>
+	template<EElementWiseUnaryOperatorType OpType, int Version>
 	bool ValidateElementWiseUnaryOperator(const NNE::FAttributeMap& AttributeMap, TConstArrayView<ENNETensorDataType> InputTypes, TConstArrayView<NNE::FSymbolicTensorShape> InputShapes)
 	{
 		return ValidateElementWiseUnaryOperatorImpl<OpType>(Version, AttributeMap, InputTypes, InputShapes);
@@ -313,7 +314,7 @@ namespace UE::NNERuntimeRDG::Private::Hlsl
 	bool RegisterElementWiseUnaryOperators(FOperatorRegistryHlsl& Registry)
 	{
 		// Note: support of a particular version is partial with respect to tensor data types (only the most typical ones are usually supported).
-#define OP(Name, Version) Registry.OpAdd({{TEXT(#Name), TEXT("Onnx")}, Version}, CreateElementWiseUnaryOperator<NNE::Internal::EElementWiseUnaryOperatorType::Name, Version>, ValidateElementWiseUnaryOperator<NNE::Internal::EElementWiseUnaryOperatorType::Name, Version>);
+#define OP(Name, Version) Registry.OpAdd({{TEXT(#Name), TEXT("Onnx")}, Version}, CreateElementWiseUnaryOperator<EElementWiseUnaryOperatorType::Name, Version>, ValidateElementWiseUnaryOperator<EElementWiseUnaryOperatorType::Name, Version>);
 		OP(Abs, 6)
 		OP(Abs, 13)
 		OP(Acos, 7)
diff --git a/Engine/Plugins/Experimental/NNERuntimeRDG/Source/NNERuntimeRDG/Private/Hlsl/NNERuntimeRDGElementWiseVariadic.cpp b/Engine/Plugins/Experimental/NNERuntimeRDG/Source/NNERuntimeRDG/Private/Hlsl/NNERuntimeRDGElementWiseVariadic.cpp
--- a/Engine/Plugins/Experimental/NNERuntimeRDG/Source/NNERuntimeRDG/Private/Hlsl/NNERuntimeRDGElementWiseVariadic.cpp
+++ b/Engine/Plugins/Experimental/NNERuntimeRDG/Source/NNERuntimeRDG/Private/Hlsl/NNERuntimeRDGElementWiseVariadic.cpp
@@ -14,12 +14,13 @@ namespace UE::NNERuntimeRDG::Private::Hlsl
 
 	using TElementWiseVariadicCS = typename UE::NNEHlslShaders::Internal::TElementWiseVariadicCS;
 	using FElementWiseVariadicConstants = UE::NNEHlslShaders::Internal::FElementWiseVariadicConstants;
+	using EElementWiseVariadicOperatorType = UE::NNEHlslShaders::Internal::EElementWiseVariadicOperatorType;
 
 	void AddOneVariadicOpPass(FRDGBuilder& GraphBuilder, 
 		TConstArrayView<FTensorRDGRef> InputTensors,
 		const FTensorRDG& OutputTensor,
 		bool OutputAsInput,
-		NNE::Internal::EElementWiseVariadicOperatorType OpType,
+		EElementWiseVariadicOperatorType OpType,
 		float Scale)
 	{
 		static_assert(FElementWiseVariadicConstants::MAX_NUM_INPUT == 4, "This algorithm need to be adapted to math the shader.");
@@ -86,7 +87,7 @@ namespace UE::NNERuntimeRDG::Private::Hlsl
 	/**
 	 * Variadic Element-wise ML operator implementation
 	 */
-	template<NNE::Internal::EElementWiseVariadicOperatorType OpType>
+	template<EElementWiseVariadicOperatorType OpType>
 	class TElementWiseVariadic : public FOperatorHlsl
 	{
 	public:
@@ -167,7 +168,7 @@ namespace UE::NNERuntimeRDG::Private::Hlsl
 					PassInputTensors[i] = InInputTensors[InputOffset + i];
 				}
 
-				if (OpType == NNE::Internal::EElementWiseVariadicOperatorType::Mean && bIsLastPass)
+				if (OpType == EElementWiseVariadicOperatorType::Mean && bIsLastPass)
 				{
 					Scale = 1.0f / InInputTensors.Num();
 				}
@@ -205,7 +206,7 @@ namespace UE::NNERuntimeRDG::Private::Hlsl
 		return bIsValid;
 	}
 
-	template<NNE::Internal::EElementWiseVariadicOperatorType OpType>
+	template<EElementWiseVariadicOperatorType OpType>
 	FOperatorHlsl* CreateElementWiseVariadicOperator()
 	{
 		return new TElementWiseVariadic<OpType>();
@@ -214,7 +215,7 @@ namespace UE::NNERuntimeRDG::Private::Hlsl
 	bool RegisterElementWiseVariadicOperators(FOperatorRegistryHlsl& Registry)
 	{
 		// Note: support of a particular version is partial with respect to tensor data types (only the most typical ones are usually supported).
-#define OP(Name, Version) Registry.OpAdd({{TEXT(#Name), TEXT("Onnx")}, Version}, CreateElementWiseVariadicOperator<NNE::Internal::EElementWiseVariadicOperatorType::Name>);
+#define OP(Name, Version) Registry.OpAdd({{TEXT(#Name), TEXT("Onnx")}, Version}, CreateElementWiseVariadicOperator<EElementWiseVariadicOperatorType::Name>);
 #define OP_ALL_VERSIONS(Name) \
 OP(Name, 6) \
 OP(Name, 8) \
diff --git a/Engine/Plugins/Experimental/NNERuntimeRDG/Source/NNERuntimeRDG/Private/NNERuntimeRDGBase.h b/Engine/Plugins/Experimental/NNERuntimeRDG/Source/NNERuntimeRDG/Private/NNERuntimeRDGBase.h
--- a/Engine/Plugins/Experimental/NNERuntimeRDG/Source/NNERuntimeRDG/Private/NNERuntimeRDGBase.h
+++ b/Engine/Plugins/Experimental/NNERuntimeRDG/Source/NNERuntimeRDG/Private/NNERuntimeRDGBase.h
@@ -5,8 +5,10 @@
 #include "NNEAttributeMap.h"
 #include "NNEAttributeValue.h"
 #include "NNEHlslShadersLog.h"
-#include "NNEModelOptimizerInterface.h"
 #include "NNERuntimeFormat.h"
+#ifdef NNE_UTILITIES_AVAILABLE
+#include "NNERuntimeRDGUtilsModelOptimizerInterface.h"
+#endif //NNE_UTILITIES_AVAILABLE
 #include "NNETensor.h"
 #include "NNETypes.h"
 #include "RenderGraphResources.h"
@@ -339,11 +341,12 @@ private:
 	TMap<FOperatorDescUnversioned, TOperatorVersionToFunctionsMap> Operators;
 };
 
+#ifdef NNE_UTILITIES_AVAILABLE
 /**
  * Validator for RDG ML operators
  */
 template<class TOperatorType>
-class TModelValidatorRDG : public UE::NNE::Internal::IModelValidator
+class TModelValidatorRDG : public NNERuntimeRDGUtils::Internal::IModelValidator
 {
 public:
 	virtual FString GetName() const 
@@ -351,7 +354,7 @@ public:
 		return TEXT("RDG Model validator");
 	}
 	
-	virtual bool ValidateModel(const FNNEModelRaw& InputModel, const UE::NNE::Internal::FOptimizerOptionsMap& Options) const override
+	virtual bool ValidateModel(const FNNEModelRaw& InputModel, const NNERuntimeRDGUtils::Internal::FOptimizerOptionsMap& Options) const override
 	{
 		FNNERuntimeFormat	Format;
 
@@ -406,5 +409,6 @@ public:
 		return true;
 	}
 };
+#endif //NNE_UTILITIES_AVAILABLE
 
 } // UE::NNERuntimeRDG::Private
\ No newline at end of file
diff --git a/Engine/Plugins/Experimental/NNERuntimeRDG/Source/NNERuntimeRDG/Private/NNERuntimeRDGHlsl.cpp b/Engine/Plugins/Experimental/NNERuntimeRDG/Source/NNERuntimeRDG/Private/NNERuntimeRDGHlsl.cpp
--- a/Engine/Plugins/Experimental/NNERuntimeRDG/Source/NNERuntimeRDG/Private/NNERuntimeRDGHlsl.cpp
+++ b/Engine/Plugins/Experimental/NNERuntimeRDG/Source/NNERuntimeRDG/Private/NNERuntimeRDGHlsl.cpp
@@ -6,11 +6,11 @@
 #include "NNEAttributeMap.h"
 #include "NNEHlslShadersLog.h"
 #include "NNEModelData.h"
-#include "NNEModelOptimizerInterface.h"
 #include "NNERuntimeRDG.h"
 #include "NNERuntimeRDGHlslHelper.h"
 #include "NNERuntimeRDGModelHlsl.h"
 #ifdef NNE_UTILITIES_AVAILABLE
+#include "NNERuntimeRDGUtilsModelOptimizerInterface.h"
 #include "NNERuntimeRDGUtilsModelOptimizer.h"
 #endif // NNE_UTILITIES_AVAILABLE
 #include "HAL/IConsoleManager.h"
@@ -231,8 +231,8 @@ TSharedPtr<UE::NNE::FSharedModelData> UNNERuntimeRDGHlslImpl::CreateModelData(co
 	}
 
 #ifdef NNE_UTILITIES_AVAILABLE
-	TUniquePtr<UE::NNE::Internal::IModelOptimizer> Optimizer = UE::NNERuntimeRDGUtils::Internal::CreateONNXToNNEModelOptimizer();
-	Optimizer->AddValidator(MakeShared<FModelValidatorHlsl>());
+	TUniquePtr<UE::NNERuntimeRDGUtils::Internal::IModelOptimizer> Optimizer = UE::NNERuntimeRDGUtils::Internal::CreateONNXToNNEModelOptimizer();
+	Optimizer->AddValidator(MakeShared<UE::NNERuntimeRDG::Private::TModelValidatorRDG<FOperatorHlsl>>());
 
 	FNNEModelRaw InputModel;
 	InputModel.Data = FileData;
diff --git a/Engine/Plugins/Experimental/NNERuntimeRDG/Source/NNERuntimeRDG/Private/NNERuntimeRDGHlslOp.h b/Engine/Plugins/Experimental/NNERuntimeRDG/Source/NNERuntimeRDG/Private/NNERuntimeRDGHlslOp.h
--- a/Engine/Plugins/Experimental/NNERuntimeRDG/Source/NNERuntimeRDG/Private/NNERuntimeRDGHlslOp.h
+++ b/Engine/Plugins/Experimental/NNERuntimeRDG/Source/NNERuntimeRDG/Private/NNERuntimeRDGHlslOp.h
@@ -21,6 +21,5 @@ namespace UE::NNERuntimeRDG::Private::Hlsl
 	  * HLSL ML operator registry
 	*/
 	typedef TOperatorRegistryRDG<FOperatorHlsl> FOperatorRegistryHlsl;
-	typedef TModelValidatorRDG<FOperatorHlsl> FModelValidatorHlsl;
 
 } // UE::NNERuntimeRDG::Private::Hlsl
\ No newline at end of file
diff --git a/Engine/Plugins/Experimental/NNERuntimeRDG/Source/NNERuntimeRDG/Private/Tests/NNERuntimeRDGUnitTestHelper.cpp b/Engine/Plugins/Experimental/NNERuntimeRDG/Source/NNERuntimeRDG/Private/Tests/NNERuntimeRDGUnitTestHelper.cpp
new file mode 100644
--- /dev/null
+++ b/Engine/Plugins/Experimental/NNERuntimeRDG/Source/NNERuntimeRDG/Private/Tests/NNERuntimeRDGUnitTestHelper.cpp
@@ -0,0 +1,140 @@
+// Copyright Epic Games, Inc. All Rights Reserved.
+
+#include "NNERuntimeRDGUnitTestHelper.h"
+#include "NNETypes.h"
+#include "NNERuntimeRDGHelperElementWiseBinary.h"
+#include "NNERuntimeRDGHelperElementWiseUnary.h"
+
+#if WITH_DEV_AUTOMATION_TESTS
+#include "Misc/AutomationTest.h"
+
+namespace UE::NNERuntimeRDG::Private::Test
+{
+
+using namespace NNE;
+using namespace NNE::Internal;
+using namespace UE::NNERuntimeRDG::Internal;
+
+bool FUnitTestBase::RunTest(const FString& Parameters) { return false; }
+
+FShapeInferenceHelperUnitTestBase::FShapeInferenceHelperUnitTestBase(const FString& InClassName, const FString& InTestName, const FString& InSourceFile, int32 InSourceLine)
+	: FUnitTestBase(InClassName), TestName(InTestName), SourceFile(InSourceFile), SourceLine(InSourceLine) {}
+
+FString FShapeInferenceHelperUnitTestBase::GetTestSourceFileName() const { return SourceFile; }
+
+int32 FShapeInferenceHelperUnitTestBase::GetTestSourceFileLine() const { return SourceLine; }
+		
+FString FShapeInferenceHelperUnitTestBase::GetBeautifiedTestName() const { return TestName; };
+
+void FShapeInferenceHelperUnitTestBase::GetTests(TArray<FString>& OutBeautifiedNames, TArray <FString>& OutTestCommands) const
+{
+	OutBeautifiedNames.Add(TestName);
+	OutTestCommands.Add(FString());
+}
+
+bool FShapeInferenceHelperUnitTestBase::RunTest(const FString& Parameter) { return false; }
+
+FTensor FShapeInferenceHelperUnitTestBase::MakeTensor(FString Name, TConstArrayView<uint32> Shape, ENNETensorDataType DataType)
+{
+	FTensorShape TensorShape = FTensorShape::Make(Shape);
+	FTensor Tensor = FTensor::Make(Name, TensorShape, DataType);
+	return Tensor;
+}
+
+FTensor FShapeInferenceHelperUnitTestBase::MakeConstTensor(FString Name, TConstArrayView<uint32> Shape, TConstArrayView<float> Data)
+{
+	FTensor Tensor = MakeTensor(Name, Shape, ENNETensorDataType::Float);
+	check(Data.Num() == Tensor.GetVolume());
+	Tensor.SetPreparedData(Data);
+	return Tensor;
+}
+
+FTensor FShapeInferenceHelperUnitTestBase::MakeConstTensorFloat16(FString Name, TConstArrayView<uint32> Shape, TConstArrayView<FFloat16> Data)
+{
+	FTensor Tensor = MakeTensor(Name, Shape, ENNETensorDataType::Half);
+	check(Data.Num() == Tensor.GetVolume());
+	Tensor.SetPreparedData(Data);
+	return Tensor;
+}
+
+FTensor FShapeInferenceHelperUnitTestBase::MakeConstTensorInt32(FString Name, TConstArrayView<uint32> Shape, TConstArrayView<int32> Data)
+{
+	FTensor Tensor = MakeTensor(Name, Shape, ENNETensorDataType::Int32);
+	check(Data.Num() == Tensor.GetVolume());
+	Tensor.SetPreparedData(Data);
+	return Tensor;
+}
+
+FTensor FShapeInferenceHelperUnitTestBase::MakeConstTensorInt64(FString Name, TConstArrayView<uint32> Shape, TConstArrayView<int64> Data)
+{
+	FTensor Tensor = MakeTensor(Name, Shape, ENNETensorDataType::Int64);
+	check(Data.Num() == Tensor.GetVolume());
+	Tensor.SetPreparedData(Data);
+	return Tensor;
+}
+
+bool FShapeInferenceHelperUnitTestBase::TestBinaryOutputIsOnlyComputedWhenItShould(NNEHlslShaders::Internal::EElementWiseBinaryOperatorType OpType)
+{
+	FTensor XC1 = MakeConstTensor(TEXT("XC1"), { 1 }, { 1.0f });
+	FTensor XC20 = MakeConstTensor(TEXT("XC20"), { 20 }, { 3.0f, 4.0f, 3.0f, 4.0f, 3.0f, 3.0f, 4.0f, 3.0f, 4.0f, 3.0f, 3.0f, 4.0f, 3.0f, 4.0f, 3.0f, 3.0f, 4.0f, 3.0f, 4.0f, 3.0f });
+	FTensor X1 = MakeTensor(TEXT("X"), { 1 });
+
+	//Tests if output tensors are only computed if both inputs are constant and below a certain size
+	
+	FTensor Y = MakeTensor(TEXT("Y"), { 1 });
+	CPUHelper::ElementWiseBinary::Apply(OpType, XC1, XC1, Y);
+	UTEST_TRUE(TEXT("Y const if both inputs are const"), Y.HasPreparedData());
+
+	Y = MakeTensor(TEXT("Y"), { 1 });
+	CPUHelper::ElementWiseBinary::Apply(OpType, X1, XC1, Y);
+	UTEST_FALSE(TEXT("Y not const when LHS not const"), Y.HasPreparedData());
+
+	Y = MakeTensor(TEXT("Y"), { 1 });
+	CPUHelper::ElementWiseBinary::Apply(OpType, XC1, X1, Y);
+	UTEST_FALSE(TEXT("Y not const when RHS not const"), Y.HasPreparedData());
+
+	Y = MakeTensor(TEXT("Y"), { 1 });
+	CPUHelper::ElementWiseBinary::Apply(OpType, X1, X1, Y);
+	UTEST_FALSE(TEXT("Y not const when LHS and RHS not const"), Y.HasPreparedData());
+
+	Y = MakeTensor(TEXT("Y"), { 20 });
+	CPUHelper::ElementWiseBinary::Apply(OpType, XC20, XC1, Y);
+	UTEST_FALSE(TEXT("Y not const when LHS too large"), Y.HasPreparedData());
+
+	Y = MakeTensor(TEXT("Y"), { 20 });
+	CPUHelper::ElementWiseBinary::Apply(OpType, XC1, XC20, Y);
+	UTEST_FALSE(TEXT("Y not const when RHS too large"), Y.HasPreparedData());
+
+	Y = MakeTensor(TEXT("Y"), { 20 });
+	CPUHelper::ElementWiseBinary::Apply(OpType, XC20, XC20, Y);
+	UTEST_FALSE(TEXT("Y not const when LHS and RHS too large"), Y.HasPreparedData());
+
+	return true;
+}
+
+bool FShapeInferenceHelperUnitTestBase::TestUnaryOutputIsOnlyComputedWhenItShould(NNEHlslShaders::Internal::EElementWiseUnaryOperatorType OpType)
+{
+	FTensor XC1 = MakeConstTensor(TEXT("XC1"), { 1 }, { 1.0f });
+	FTensor XC20 = MakeConstTensor(TEXT("XC20"), { 20 }, { 3.0f, 4.0f, 3.0f, 4.0f, 3.0f, 3.0f, 4.0f, 3.0f, 4.0f, 3.0f, 3.0f, 4.0f, 3.0f, 4.0f, 3.0f, 3.0f, 4.0f, 3.0f, 4.0f, 3.0f });
+	FTensor X1 = MakeTensor(TEXT("X"), { 1 });
+
+	//Tests if output tensors are only computed if inputs are constant and below a certain size
+
+	FTensor Y = MakeTensor(TEXT("Y"), { 1 });
+	CPUHelper::ElementWiseUnary::Apply(OpType, XC1, 1.0f, 1.0f, 1.0f, Y);
+	UTEST_TRUE(TEXT("Y const if input is const"), Y.HasPreparedData());
+
+	Y = MakeTensor(TEXT("Y"), { 1 });
+	CPUHelper::ElementWiseUnary::Apply(OpType, X1, 1.0f, 1.0f, 1.0f, Y);
+	UTEST_FALSE(TEXT("Y not const if input is not const"), Y.HasPreparedData());
+
+	Y = MakeTensor(TEXT("Y"), { 20 });
+	CPUHelper::ElementWiseUnary::Apply(OpType, XC20, 1.0f, 1.0f, 1.0f, Y);
+	UTEST_FALSE(TEXT("Y not const if input is too large "), Y.HasPreparedData());
+
+	return true;
+}
+
+} // namespace UE::NNERuntimeRDG::Private::Test
+
+#endif //WITH_DEV_AUTOMATION_TESTS
diff --git a/Engine/Plugins/Experimental/NNERuntimeRDG/Source/NNERuntimeRDG/Private/Tests/NNERuntimeRDGUnitTestHelper.h b/Engine/Plugins/Experimental/NNERuntimeRDG/Source/NNERuntimeRDG/Private/Tests/NNERuntimeRDGUnitTestHelper.h
new file mode 100644
--- /dev/null
+++ b/Engine/Plugins/Experimental/NNERuntimeRDG/Source/NNERuntimeRDG/Private/Tests/NNERuntimeRDGUnitTestHelper.h
@@ -0,0 +1,64 @@
+// Copyright Epic Games, Inc. All Rights Reserved.
+
+#pragma once
+
+#include "CoreMinimal.h"
+#include "NNETensor.h"
+#include "NNEHlslShadersOperator.h"
+
+#if WITH_DEV_AUTOMATION_TESTS
+#include "Misc/AutomationTest.h"
+
+namespace UE::NNERuntimeRDG::Private::Test
+{
+	IMPLEMENT_SIMPLE_AUTOMATION_TEST_PRIVATE(FUnitTestBase, FAutomationTestBase, "System.Engine.MachineLearning.NNE.RDG.UnitTest.Base", EAutomationTestFlags_ApplicationContextMask | EAutomationTestFlags_FeatureMask | EAutomationTestFlags::EngineFilter, __FILE__, __LINE__)
+
+	class FShapeInferenceHelperUnitTestBase : public FUnitTestBase
+	{
+		//Automation tests
+		FString TestName;
+		FString SourceFile;
+		int32 SourceLine;
+		
+	public:
+		FShapeInferenceHelperUnitTestBase(const FString& InClassName, const FString& InTestName, const FString& InSourceFile, int32 InSourceLine);
+		virtual ~FShapeInferenceHelperUnitTestBase() {}
+		virtual FString GetTestSourceFileName() const override;
+		virtual int32 GetTestSourceFileLine() const override;
+		
+	protected:
+		virtual FString GetBeautifiedTestName() const override;
+		virtual void GetTests(TArray<FString>& OutBeautifiedNames, TArray <FString>& OutTestCommands) const override;
+		virtual bool RunTest(const FString& Parameter) override;
+
+		//Shape inference unit test helpers
+		static NNE::Internal::FTensor MakeTensor(FString Name, TConstArrayView<uint32> Shape, ENNETensorDataType DataType = ENNETensorDataType::Float);
+		static NNE::Internal::FTensor MakeConstTensor(FString Name, TConstArrayView<uint32> Shape, TConstArrayView<float> Data);
+		static NNE::Internal::FTensor MakeConstTensorFloat16(FString Name, TConstArrayView<uint32> Shape, TConstArrayView<FFloat16> Data);
+		static NNE::Internal::FTensor MakeConstTensorInt32(FString Name, TConstArrayView<uint32> Shape, TConstArrayView<int32> Data);
+		static NNE::Internal::FTensor MakeConstTensorInt64(FString Name, TConstArrayView<uint32> Shape, TConstArrayView<int64> Data);
+		bool TestUnaryOutputIsOnlyComputedWhenItShould(NNEHlslShaders::Internal::EElementWiseUnaryOperatorType OpType);
+		bool TestBinaryOutputIsOnlyComputedWhenItShould(NNEHlslShaders::Internal::EElementWiseBinaryOperatorType OpType);
+	};
+
+	#define IMPLEMENT_NNE_SHAPEINFERENCEHELPER_UNIT_AUTOMATION_TEST( TClass, PrettyName ) \
+	class TClass : public FShapeInferenceHelperUnitTestBase \
+	{ \
+	public: \
+		TClass(const FString& InClassName, const FString& InTestName, const FString& InSourceFile, int32 InSourceLine) \
+			: FShapeInferenceHelperUnitTestBase(InClassName, InTestName, InSourceFile, InSourceLine) {} \
+		virtual ~TClass() {} \
+    \
+	protected: \
+		virtual bool RunTest(const FString& Parameter) override; \
+	}; \
+	namespace\
+	{\
+		TClass TClass##AutomationTestInstance( TEXT(#TClass), PrettyName, __FILE__, __LINE__ );\
+	}
+
+
+
+} // namespace UE::NNERuntimeRDG::Private::Test
+
+#endif //WITH_DEV_AUTOMATION_TESTS
diff --git a/Engine/Plugins/Experimental/NNERuntimeRDG/Source/NNERuntimeRDG/Private/Tests/NNERuntimeRDGUnitTestOpBinary.cpp b/Engine/Plugins/Experimental/NNERuntimeRDG/Source/NNERuntimeRDG/Private/Tests/NNERuntimeRDGUnitTestOpBinary.cpp
new file mode 100644
--- /dev/null
+++ b/Engine/Plugins/Experimental/NNERuntimeRDG/Source/NNERuntimeRDG/Private/Tests/NNERuntimeRDGUnitTestOpBinary.cpp
@@ -0,0 +1,361 @@
+// Copyright Epic Games, Inc. All Rights Reserved.
+
+#include "CoreMinimal.h"
+#include "NNETensor.h"
+#include "NNETypes.h"
+#include "NNERuntimeRDGHelperElementWiseBinary.h"
+#include "NNERuntimeRDGUnitTestHelper.h"
+
+
+#if WITH_DEV_AUTOMATION_TESTS
+#include "Misc/AutomationTest.h"
+#endif
+
+namespace UE::NNERuntimeRDG::Private::Test::BinaryOp
+{
+
+#if WITH_DEV_AUTOMATION_TESTS
+	
+	using namespace NNE;
+	using namespace NNE::Internal;
+	using namespace UE::NNERuntimeRDG::Internal;
+	using EElementWiseBinaryOperatorType = NNEHlslShaders::Internal::EElementWiseBinaryOperatorType;
+
+	IMPLEMENT_NNE_SHAPEINFERENCEHELPER_UNIT_AUTOMATION_TEST(FElementWiseBinaryCPUHelperAdd, "System.Engine.MachineLearning.NNE.RDG.UnitTest.BinaryHelper.Add");
+	bool FElementWiseBinaryCPUHelperAdd::RunTest(const FString& Parameter)
+	{
+		FTensor XC1 = MakeConstTensor(TEXT("XC1"), { 1 }, { 1.0f });
+		FTensor XC1x2 = MakeConstTensor(TEXT("XC1x2"), { 1,2 }, { 1.0f, 2.0f });
+		FTensor XC2x1 = MakeConstTensor(TEXT("XC2x1"), { 2,1 }, { 3.0f, 4.0f });
+
+		if (!TestBinaryOutputIsOnlyComputedWhenItShould(EElementWiseBinaryOperatorType::Add))
+		{
+			return false;
+		}
+
+		//Test output tensor math is correct when both inputs are constant including broadcasting
+		FTensor Y = MakeTensor(TEXT("Y"), { 1 });
+		CPUHelper::ElementWiseBinary::Apply(EElementWiseBinaryOperatorType::Add, XC1, XC1, Y);
+		UTEST_EQUAL(TEXT("XC1+XC1"), Y.GetPreparedData<float>()[0], 1.0f + 1.0f);
+
+		Y = MakeTensor(TEXT("Y"), { 1,2 });
+		CPUHelper::ElementWiseBinary::Apply(EElementWiseBinaryOperatorType::Add, XC1x2, XC1, Y);
+		UTEST_EQUAL(TEXT("XC1x2+XC1[0]"), Y.GetPreparedData<float>()[0], 1.0f + 1.0f);
+		UTEST_EQUAL(TEXT("XC1x2+XC1[1]"), Y.GetPreparedData<float>()[1], 2.0f + 1.0f);
+
+		Y = MakeTensor(TEXT("Y"), { 1,2 });
+		CPUHelper::ElementWiseBinary::Apply(EElementWiseBinaryOperatorType::Add, XC1x2, XC1x2, Y);
+		UTEST_EQUAL(TEXT("XC1x2+XC1x2[0]"), Y.GetPreparedData<float>()[0], 1.0f + 1.0f);
+		UTEST_EQUAL(TEXT("XC1x2+XC1x2[1]"), Y.GetPreparedData<float>()[1], 2.0f + 2.0f);
+
+		Y = MakeTensor(TEXT("Y"), { 2,2 });
+		CPUHelper::ElementWiseBinary::Apply(EElementWiseBinaryOperatorType::Add, XC1x2, XC2x1, Y);
+		UTEST_EQUAL(TEXT("XC1x2+XC2x1[0]"), Y.GetPreparedData<float>()[0], 1.0f + 3.0f);
+		UTEST_EQUAL(TEXT("XC1x2+XC2x1[1]"), Y.GetPreparedData<float>()[1], 2.0f + 3.0f);
+		UTEST_EQUAL(TEXT("XC1x2+XC2x1[2]"), Y.GetPreparedData<float>()[2], 1.0f + 4.0f);
+		UTEST_EQUAL(TEXT("XC1x2+XC2x1[3]"), Y.GetPreparedData<float>()[3], 2.0f + 4.0f);
+
+		return true;
+	}
+
+	IMPLEMENT_NNE_SHAPEINFERENCEHELPER_UNIT_AUTOMATION_TEST(FElementWiseBinaryCPUHelperDiv, "System.Engine.MachineLearning.NNE.RDG.UnitTest.BinaryHelper.Div");
+	bool FElementWiseBinaryCPUHelperDiv::RunTest(const FString& Parameter)
+	{
+		FTensor XC1 = MakeConstTensor(TEXT("XC1"), { 1 }, { 1.0f });
+		FTensor XC1x2 = MakeConstTensor(TEXT("XC1x2"), { 1,2 }, { 1.0f, 2.0f });
+		FTensor XC2x1 = MakeConstTensor(TEXT("XC2x1"), { 2,1 }, { 3.0f, 4.0f });
+
+		if (!TestBinaryOutputIsOnlyComputedWhenItShould(EElementWiseBinaryOperatorType::Div))
+		{
+			return false;
+		}
+
+		//Test output tensor math is correct when both inputs are constant including broadcasting
+		FTensor Y = MakeTensor(TEXT("Y"), { 1 });
+		CPUHelper::ElementWiseBinary::Apply(EElementWiseBinaryOperatorType::Div, XC1, XC1, Y);
+		UTEST_EQUAL(TEXT("XC1/XC1"), Y.GetPreparedData<float>()[0], 1.0f / 1.0f);
+
+		Y = MakeTensor(TEXT("Y"), { 1,2 });
+		CPUHelper::ElementWiseBinary::Apply(EElementWiseBinaryOperatorType::Div, XC1x2, XC1, Y);
+		UTEST_EQUAL(TEXT("XC1x2/XC1[0]"), Y.GetPreparedData<float>()[0], 1.0f / 1.0f);
+		UTEST_EQUAL(TEXT("XC1x2/XC1[1]"), Y.GetPreparedData<float>()[1], 2.0f / 1.0f);
+
+		Y = MakeTensor(TEXT("Y"), { 1,2 });
+		CPUHelper::ElementWiseBinary::Apply(EElementWiseBinaryOperatorType::Div, XC1x2, XC1x2, Y);
+		UTEST_EQUAL(TEXT("XC1x2/XC1x2[0]"), Y.GetPreparedData<float>()[0], 1.0f / 1.0f);
+		UTEST_EQUAL(TEXT("XC1x2/XC1x2[1]"), Y.GetPreparedData<float>()[1], 2.0f / 2.0f);
+
+		Y = MakeTensor(TEXT("Y"), { 2,2 });
+		CPUHelper::ElementWiseBinary::Apply(EElementWiseBinaryOperatorType::Div, XC1x2, XC2x1, Y);
+		UTEST_EQUAL(TEXT("XC1x2/XC2x1[0]"), Y.GetPreparedData<float>()[0], 1.0f / 3.0f);
+		UTEST_EQUAL(TEXT("XC1x2/XC2x1[1]"), Y.GetPreparedData<float>()[1], 2.0f / 3.0f);
+		UTEST_EQUAL(TEXT("XC1x2/XC2x1[2]"), Y.GetPreparedData<float>()[2], 1.0f / 4.0f);
+		UTEST_EQUAL(TEXT("XC1x2/XC2x1[3]"), Y.GetPreparedData<float>()[3], 2.0f / 4.0f);
+
+		return true;
+	};
+
+	IMPLEMENT_NNE_SHAPEINFERENCEHELPER_UNIT_AUTOMATION_TEST(FElementWiseBinaryCPUHelperMul, "System.Engine.MachineLearning.NNE.RDG.UnitTest.BinaryHelper.Mul");
+	bool FElementWiseBinaryCPUHelperMul::RunTest(const FString& Parameter)
+	{
+		FTensor XC1 = MakeConstTensor(TEXT("XC1"), { 1 }, { 1.0f });
+		FTensor XC1x2 = MakeConstTensor(TEXT("XC1x2"), { 1,2 }, { 1.0f, 2.0f });
+		FTensor XC2x1 = MakeConstTensor(TEXT("XC2x1"), { 2,1 }, { 3.0f, 4.0f });
+
+		if (!TestBinaryOutputIsOnlyComputedWhenItShould(EElementWiseBinaryOperatorType::Mul))
+		{
+			return false;
+		}
+
+		//Test output tensor math is correct when both inputs are constant including broadcasting
+		FTensor Y = MakeTensor(TEXT("Y"), { 1 });
+		CPUHelper::ElementWiseBinary::Apply(EElementWiseBinaryOperatorType::Mul, XC1, XC1, Y);
+		UTEST_EQUAL(TEXT("XC1*XC1"), Y.GetPreparedData<float>()[0], 1.0f * 1.0f);
+
+		Y = MakeTensor(TEXT("Y"), { 1,2 });
+		CPUHelper::ElementWiseBinary::Apply(EElementWiseBinaryOperatorType::Mul, XC1x2, XC1, Y);
+		UTEST_EQUAL(TEXT("XC1x2*XC1[0]"), Y.GetPreparedData<float>()[0], 1.0f * 1.0f);
+		UTEST_EQUAL(TEXT("XC1x2*XC1[1]"), Y.GetPreparedData<float>()[1], 2.0f * 1.0f);
+
+		Y = MakeTensor(TEXT("Y"), { 1,2 });
+		CPUHelper::ElementWiseBinary::Apply(EElementWiseBinaryOperatorType::Mul, XC1x2, XC1x2, Y);
+		UTEST_EQUAL(TEXT("XC1x2*XC1x2[0]"), Y.GetPreparedData<float>()[0], 1.0f * 1.0f);
+		UTEST_EQUAL(TEXT("XC1x2*XC1x2[1]"), Y.GetPreparedData<float>()[1], 2.0f * 2.0f);
+
+		Y = MakeTensor(TEXT("Y"), { 2,2 });
+		CPUHelper::ElementWiseBinary::Apply(EElementWiseBinaryOperatorType::Mul, XC1x2, XC2x1, Y);
+		UTEST_EQUAL(TEXT("XC1x2*XC2x1[0]"), Y.GetPreparedData<float>()[0], 1.0f * 3.0f);
+		UTEST_EQUAL(TEXT("XC1x2*XC2x1[1]"), Y.GetPreparedData<float>()[1], 2.0f * 3.0f);
+		UTEST_EQUAL(TEXT("XC1x2*XC2x1[2]"), Y.GetPreparedData<float>()[2], 1.0f * 4.0f);
+		UTEST_EQUAL(TEXT("XC1x2*XC2x1[3]"), Y.GetPreparedData<float>()[3], 2.0f * 4.0f);
+
+		return true;
+	};
+
+	IMPLEMENT_NNE_SHAPEINFERENCEHELPER_UNIT_AUTOMATION_TEST(FElementWiseBinaryCPUHelperMulInt32, "System.Engine.MachineLearning.NNE.RDG.UnitTest.BinaryHelper.MulInt32");
+	bool FElementWiseBinaryCPUHelperMulInt32::RunTest(const FString& Parameter)
+	{
+		FTensor XC1 = MakeConstTensorInt32(TEXT("XC1"), { 1 }, { 1 });
+		FTensor XC1x2 = MakeConstTensorInt32(TEXT("XC1x2"), { 1,2 }, { 1, 2 });
+		FTensor XC2x1 = MakeConstTensorInt32(TEXT("XC2x1"), { 2,1 }, { 3, 4 });
+
+		if (!TestBinaryOutputIsOnlyComputedWhenItShould(EElementWiseBinaryOperatorType::Mul))
+		{
+			return false;
+		}
+
+		//Test output tensor math is correct when both inputs are constant including broadcasting
+		FTensor Y = MakeTensor(TEXT("Y"), { 1 }, ENNETensorDataType::Int32);
+		CPUHelper::ElementWiseBinary::Apply(EElementWiseBinaryOperatorType::Mul, XC1, XC1, Y);
+		UTEST_EQUAL(TEXT("XC1*XC1"), Y.GetPreparedData<int32>()[0], 1 * 1);
+
+		Y = MakeTensor(TEXT("Y"), { 1,2 }, ENNETensorDataType::Int32);
+		CPUHelper::ElementWiseBinary::Apply(EElementWiseBinaryOperatorType::Mul, XC1x2, XC1, Y);
+		UTEST_EQUAL(TEXT("XC1x2*XC1[0]"), Y.GetPreparedData<int32>()[0], 1 * 1);
+		UTEST_EQUAL(TEXT("XC1x2*XC1[1]"), Y.GetPreparedData<int32>()[1], 2 * 1);
+
+		Y = MakeTensor(TEXT("Y"), { 1,2 }, ENNETensorDataType::Int32);
+		CPUHelper::ElementWiseBinary::Apply(EElementWiseBinaryOperatorType::Mul, XC1x2, XC1x2, Y);
+		UTEST_EQUAL(TEXT("XC1x2*XC1x2[0]"), Y.GetPreparedData<int32>()[0], 1 * 1);
+		UTEST_EQUAL(TEXT("XC1x2*XC1x2[1]"), Y.GetPreparedData<int32>()[1], 2 * 2);
+
+		Y = MakeTensor(TEXT("Y"), { 2,2 }, ENNETensorDataType::Int32);
+		CPUHelper::ElementWiseBinary::Apply(EElementWiseBinaryOperatorType::Mul, XC1x2, XC2x1, Y);
+		UTEST_EQUAL(TEXT("XC1x2*XC2x1[0]"), Y.GetPreparedData<int32>()[0], 1 * 3);
+		UTEST_EQUAL(TEXT("XC1x2*XC2x1[1]"), Y.GetPreparedData<int32>()[1], 2 * 3);
+		UTEST_EQUAL(TEXT("XC1x2*XC2x1[2]"), Y.GetPreparedData<int32>()[2], 1 * 4);
+		UTEST_EQUAL(TEXT("XC1x2*XC2x1[3]"), Y.GetPreparedData<int32>()[3], 2 * 4);
+
+		return true;
+	};
+
+	IMPLEMENT_NNE_SHAPEINFERENCEHELPER_UNIT_AUTOMATION_TEST(FElementWiseBinaryCPUHelperMulInt64, "System.Engine.MachineLearning.NNE.RDG.UnitTest.BinaryHelper.MulInt64");
+	bool FElementWiseBinaryCPUHelperMulInt64::RunTest(const FString& Parameter)
+	{
+		FTensor XC1 = MakeConstTensorInt64(TEXT("XC1"), { 1 }, { 1 });
+		FTensor XC1x2 = MakeConstTensorInt64(TEXT("XC1x2"), { 1,2 }, { 1, 2 });
+		FTensor XC2x1 = MakeConstTensorInt64(TEXT("XC2x1"), { 2,1 }, { 3, 4 });
+
+		if (!TestBinaryOutputIsOnlyComputedWhenItShould(EElementWiseBinaryOperatorType::Mul))
+		{
+			return false;
+		}
+
+		//Test output tensor math is correct when both inputs are constant including broadcasting
+		FTensor Y = MakeTensor(TEXT("Y"), { 1 }, ENNETensorDataType::Int64);
+		CPUHelper::ElementWiseBinary::Apply(EElementWiseBinaryOperatorType::Mul, XC1, XC1, Y);
+		UTEST_EQUAL(TEXT("XC1*XC1"), Y.GetPreparedData<int64>()[0], (int64)1 * 1);
+
+		Y = MakeTensor(TEXT("Y"), { 1,2 }, ENNETensorDataType::Int64);
+		CPUHelper::ElementWiseBinary::Apply(EElementWiseBinaryOperatorType::Mul, XC1x2, XC1, Y);
+		UTEST_EQUAL(TEXT("XC1x2*XC1[0]"), Y.GetPreparedData<int64>()[0], (int64)1 * 1);
+		UTEST_EQUAL(TEXT("XC1x2*XC1[1]"), Y.GetPreparedData<int64>()[1], (int64)2 * 1);
+
+		Y = MakeTensor(TEXT("Y"), { 1,2 }, ENNETensorDataType::Int64);
+		CPUHelper::ElementWiseBinary::Apply(EElementWiseBinaryOperatorType::Mul, XC1x2, XC1x2, Y);
+		UTEST_EQUAL(TEXT("XC1x2*XC1x2[0]"), Y.GetPreparedData<int64>()[0], (int64)1 * 1);
+		UTEST_EQUAL(TEXT("XC1x2*XC1x2[1]"), Y.GetPreparedData<int64>()[1], (int64)2 * 2);
+
+		Y = MakeTensor(TEXT("Y"), { 2,2 }, ENNETensorDataType::Int64);
+		CPUHelper::ElementWiseBinary::Apply(EElementWiseBinaryOperatorType::Mul, XC1x2, XC2x1, Y);
+		UTEST_EQUAL(TEXT("XC1x2*XC2x1[0]"), Y.GetPreparedData<int64>()[0], (int64)1 * 3);
+		UTEST_EQUAL(TEXT("XC1x2*XC2x1[1]"), Y.GetPreparedData<int64>()[1], (int64)2 * 3);
+		UTEST_EQUAL(TEXT("XC1x2*XC2x1[2]"), Y.GetPreparedData<int64>()[2], (int64)1 * 4);
+		UTEST_EQUAL(TEXT("XC1x2*XC2x1[3]"), Y.GetPreparedData<int64>()[3], (int64)2 * 4);
+
+		return true;
+	};
+
+	IMPLEMENT_NNE_SHAPEINFERENCEHELPER_UNIT_AUTOMATION_TEST(FElementWiseBinaryCPUHelperSub, "System.Engine.MachineLearning.NNE.RDG.UnitTest.BinaryHelper.Sub");
+	bool FElementWiseBinaryCPUHelperSub::RunTest(const FString& Parameter)
+	{
+		FTensor XC1 = MakeConstTensor(TEXT("XC1"), { 1 }, { 1.0f });
+		FTensor XC1x2 = MakeConstTensor(TEXT("XC1x2"), { 1,2 }, { 1.0f, 2.0f });
+		FTensor XC2x1 = MakeConstTensor(TEXT("XC2x1"), { 2,1 }, { 3.0f, 4.0f });
+
+		if (!TestBinaryOutputIsOnlyComputedWhenItShould(EElementWiseBinaryOperatorType::Sub))
+		{
+			return false;
+		}
+
+		//Test output tensor math is correct when both inputs are constant including broadcasting
+		FTensor Y = MakeTensor(TEXT("Y"), { 1 });
+		CPUHelper::ElementWiseBinary::Apply(EElementWiseBinaryOperatorType::Sub, XC1, XC1, Y);
+		UTEST_EQUAL(TEXT("XC1-XC1"), Y.GetPreparedData<float>()[0], 1.0f - 1.0f);
+
+		Y = MakeTensor(TEXT("Y"), { 1,2 });
+		CPUHelper::ElementWiseBinary::Apply(EElementWiseBinaryOperatorType::Sub, XC1x2, XC1, Y);
+		UTEST_EQUAL(TEXT("XC1x2-XC1[0]"), Y.GetPreparedData<float>()[0], 1.0f - 1.0f);
+		UTEST_EQUAL(TEXT("XC1x2-XC1[1]"), Y.GetPreparedData<float>()[1], 2.0f - 1.0f);
+
+		Y = MakeTensor(TEXT("Y"), { 1,2 });
+		CPUHelper::ElementWiseBinary::Apply(EElementWiseBinaryOperatorType::Sub, XC1x2, XC1x2, Y);
+		UTEST_EQUAL(TEXT("XC1x2-XC1x2[0]"), Y.GetPreparedData<float>()[0], 1.0f - 1.0f);
+		UTEST_EQUAL(TEXT("XC1x2-XC1x2[1]"), Y.GetPreparedData<float>()[1], 2.0f - 2.0f);
+
+		Y = MakeTensor(TEXT("Y"), { 2,2 });
+		CPUHelper::ElementWiseBinary::Apply(EElementWiseBinaryOperatorType::Sub, XC1x2, XC2x1, Y);
+		UTEST_EQUAL(TEXT("XC1x2-XC2x1[0]"), Y.GetPreparedData<float>()[0], 1.0f - 3.0f);
+		UTEST_EQUAL(TEXT("XC1x2-XC2x1[1]"), Y.GetPreparedData<float>()[1], 2.0f - 3.0f);
+		UTEST_EQUAL(TEXT("XC1x2-XC2x1[2]"), Y.GetPreparedData<float>()[2], 1.0f - 4.0f);
+		UTEST_EQUAL(TEXT("XC1x2-XC2x1[3]"), Y.GetPreparedData<float>()[3], 2.0f - 4.0f);
+
+		return true;
+	};
+
+	IMPLEMENT_NNE_SHAPEINFERENCEHELPER_UNIT_AUTOMATION_TEST(FElementWiseBinaryCPUHelperMod, "System.Engine.MachineLearning.NNE.RDG.UnitTest.BinaryHelper.Mod");
+	bool FElementWiseBinaryCPUHelperMod::RunTest(const FString& Parameter)
+	{
+		FTensor XC1 = MakeConstTensor(TEXT("XC1"), { 1 }, { 1.0f });
+		FTensor XC1x2 = MakeConstTensor(TEXT("XC1x2"), { 1,2 }, { 1.0f, 2.0f });
+		FTensor XC2x1 = MakeConstTensor(TEXT("XC2x1"), { 2,1 }, { 3.0f, 4.0f });
+
+		if (!TestBinaryOutputIsOnlyComputedWhenItShould(EElementWiseBinaryOperatorType::Mod))
+		{
+			return false;
+		}
+
+		//Test output tensor math is correct when both inputs are constant including broadcasting
+		FTensor Y = MakeTensor(TEXT("Y"), { 1 });
+		CPUHelper::ElementWiseBinary::Apply(EElementWiseBinaryOperatorType::Mod, XC1, XC1, Y);
+		UTEST_EQUAL(TEXT("XC1%XC1"), Y.GetPreparedData<float>()[0], FMath::Fmod(1.0f, 1.0f));
+
+		Y = MakeTensor(TEXT("Y"), { 1,2 });
+		CPUHelper::ElementWiseBinary::Apply(EElementWiseBinaryOperatorType::Mod, XC1x2, XC1, Y);
+		UTEST_EQUAL(TEXT("XC1x2%XC1[0]"), Y.GetPreparedData<float>()[0], FMath::Fmod(1.0f, 1.0f));
+		UTEST_EQUAL(TEXT("XC1x2%XC1[1]"), Y.GetPreparedData<float>()[1], FMath::Fmod(2.0f, 1.0f));
+
+		Y = MakeTensor(TEXT("Y"), { 1,2 });
+		CPUHelper::ElementWiseBinary::Apply(EElementWiseBinaryOperatorType::Mod, XC1x2, XC1x2, Y);
+		UTEST_EQUAL(TEXT("XC1x2%XC1x2[0]"), Y.GetPreparedData<float>()[0], FMath::Fmod(1.0f, 1.0f));
+		UTEST_EQUAL(TEXT("XC1x2%XC1x2[1]"), Y.GetPreparedData<float>()[1], FMath::Fmod(2.0f, 2.0f));
+
+		Y = MakeTensor(TEXT("Y"), { 2,2 });
+		CPUHelper::ElementWiseBinary::Apply(EElementWiseBinaryOperatorType::Mod, XC1x2, XC2x1, Y);
+		UTEST_EQUAL(TEXT("XC1x2%XC2x1[0]"), Y.GetPreparedData<float>()[0], FMath::Fmod(1.0f, 3.0f));
+		UTEST_EQUAL(TEXT("XC1x2%XC2x1[1]"), Y.GetPreparedData<float>()[1], FMath::Fmod(2.0f, 3.0f));
+		UTEST_EQUAL(TEXT("XC1x2%XC2x1[2]"), Y.GetPreparedData<float>()[2], FMath::Fmod(1.0f, 4.0f));
+		UTEST_EQUAL(TEXT("XC1x2%XC2x1[3]"), Y.GetPreparedData<float>()[3], FMath::Fmod(2.0f, 4.0f));
+
+		return true;
+	};
+
+	IMPLEMENT_NNE_SHAPEINFERENCEHELPER_UNIT_AUTOMATION_TEST(FElementWiseBinaryCPUHelperPow, "System.Engine.MachineLearning.NNE.RDG.UnitTest.BinaryHelper.Pow");
+	bool FElementWiseBinaryCPUHelperPow::RunTest(const FString& Parameter)
+	{
+		FTensor XC1 = MakeConstTensor(TEXT("XC1"), { 1 }, { 1.0f });
+		FTensor XC1x2 = MakeConstTensor(TEXT("XC1x2"), { 1,2 }, { 1.0f, 2.0f });
+		FTensor XC2x1 = MakeConstTensor(TEXT("XC2x1"), { 2,1 }, { 3.0f, 4.0f });
+
+		if (!TestBinaryOutputIsOnlyComputedWhenItShould(EElementWiseBinaryOperatorType::Pow))
+		{
+			return false;
+		}
+
+		//Test output tensor math is correct when both inputs are constant including broadcasting
+		FTensor Y = MakeTensor(TEXT("Y"), { 1 });
+		CPUHelper::ElementWiseBinary::Apply(EElementWiseBinaryOperatorType::Pow, XC1, XC1, Y);
+		UTEST_EQUAL(TEXT("XC1^XC1"), Y.GetPreparedData<float>()[0], FMath::Pow(1.0f, 1.0f));
+
+		Y = MakeTensor(TEXT("Y"), { 1,2 });
+		CPUHelper::ElementWiseBinary::Apply(EElementWiseBinaryOperatorType::Pow, XC1x2, XC1, Y);
+		UTEST_EQUAL(TEXT("XC1x2^XC1[0]"), Y.GetPreparedData<float>()[0], FMath::Pow(1.0f, 1.0f));
+		UTEST_EQUAL(TEXT("XC1x2^XC1[1]"), Y.GetPreparedData<float>()[1], FMath::Pow(2.0f, 1.0f));
+
+		Y = MakeTensor(TEXT("Y"), { 1,2 });
+		CPUHelper::ElementWiseBinary::Apply(EElementWiseBinaryOperatorType::Pow, XC1x2, XC1x2, Y);
+		UTEST_EQUAL(TEXT("XC1x2^XC1x2[0]"), Y.GetPreparedData<float>()[0], FMath::Pow(1.0f, 1.0f));
+		UTEST_EQUAL(TEXT("XC1x2^XC1x2[1]"), Y.GetPreparedData<float>()[1], FMath::Pow(2.0f, 2.0f));
+
+		Y = MakeTensor(TEXT("Y"), { 2,2 });
+		CPUHelper::ElementWiseBinary::Apply(EElementWiseBinaryOperatorType::Pow, XC1x2, XC2x1, Y);
+		UTEST_EQUAL(TEXT("XC1x2^XC2x1[0]"), Y.GetPreparedData<float>()[0], FMath::Pow(1.0f, 3.0f));
+		UTEST_EQUAL(TEXT("XC1x2^XC2x1[1]"), Y.GetPreparedData<float>()[1], FMath::Pow(2.0f, 3.0f));
+		UTEST_EQUAL(TEXT("XC1x2^XC2x1[2]"), Y.GetPreparedData<float>()[2], FMath::Pow(1.0f, 4.0f));
+		UTEST_EQUAL(TEXT("XC1x2^XC2x1[3]"), Y.GetPreparedData<float>()[3], FMath::Pow(2.0f, 4.0f));
+
+		return true;
+	};
+
+	IMPLEMENT_NNE_SHAPEINFERENCEHELPER_UNIT_AUTOMATION_TEST(FElementWiseBinaryCPUHelperPrelu, "System.Engine.MachineLearning.NNE.RDG.UnitTest.BinaryHelper.Prelu");
+	bool FElementWiseBinaryCPUHelperPrelu::RunTest(const FString& Parameter)
+	{
+		FTensor XC1 = MakeConstTensor(TEXT("XC1"), { 1 }, { 1.0f });
+		FTensor XC1x2 = MakeConstTensor(TEXT("XC1x2"), { 1,2 }, { 1.0f, 2.0f });
+		FTensor XC2x1 = MakeConstTensor(TEXT("XC2x1"), { 2,1 }, { 3.0f, 4.0f });
+
+		if (!TestBinaryOutputIsOnlyComputedWhenItShould(EElementWiseBinaryOperatorType::Prelu))
+		{
+			return false;
+		}
+
+		auto Prelu = [](float X, float Y) { return (X < 0.0f) ? (Y * X) : X; };
+
+		//Test output tensor math is correct when both inputs are constant including broadcasting
+		FTensor Y = MakeTensor(TEXT("Y"), { 1 });
+		CPUHelper::ElementWiseBinary::Apply(EElementWiseBinaryOperatorType::Prelu, XC1, XC1, Y);
+		UTEST_EQUAL(TEXT("XC1_XC1"), Y.GetPreparedData<float>()[0], Prelu(1.0f, 1.0f));
+
+		Y = MakeTensor(TEXT("Y"), { 1,2 });
+		CPUHelper::ElementWiseBinary::Apply(EElementWiseBinaryOperatorType::Prelu, XC1x2, XC1, Y);
+		UTEST_EQUAL(TEXT("XC1x2_XC1[0]"), Y.GetPreparedData<float>()[0], Prelu(1.0f, 1.0f));
+		UTEST_EQUAL(TEXT("XC1x2_XC1[1]"), Y.GetPreparedData<float>()[1], Prelu(2.0f, 1.0f));
+
+		Y = MakeTensor(TEXT("Y"), { 1,2 });
+		CPUHelper::ElementWiseBinary::Apply(EElementWiseBinaryOperatorType::Prelu, XC1x2, XC1x2, Y);
+		UTEST_EQUAL(TEXT("XC1x2_XC1x2[0]"), Y.GetPreparedData<float>()[0], Prelu(1.0f, 1.0f));
+		UTEST_EQUAL(TEXT("XC1x2_XC1x2[1]"), Y.GetPreparedData<float>()[1], Prelu(2.0f, 2.0f));
+
+		Y = MakeTensor(TEXT("Y"), { 2,2 });
+		CPUHelper::ElementWiseBinary::Apply(EElementWiseBinaryOperatorType::Prelu, XC1x2, XC2x1, Y);
+		UTEST_EQUAL(TEXT("XC1x2_XC2x1[0]"), Y.GetPreparedData<float>()[0], Prelu(1.0f, 3.0f));
+		UTEST_EQUAL(TEXT("XC1x2_XC2x1[1]"), Y.GetPreparedData<float>()[1], Prelu(2.0f, 3.0f));
+		UTEST_EQUAL(TEXT("XC1x2_XC2x1[2]"), Y.GetPreparedData<float>()[2], Prelu(1.0f, 4.0f));
+		UTEST_EQUAL(TEXT("XC1x2_XC2x1[3]"), Y.GetPreparedData<float>()[3], Prelu(2.0f, 4.0f));
+
+		return true;
+	};
+
+#endif //WITH_DEV_AUTOMATION_TESTS
+
+} // namespace UE::NNERuntimeRDG::Private::Test::BinaryOp
diff --git a/Engine/Plugins/Experimental/NNERuntimeRDG/Source/NNERuntimeRDG/Private/Tests/NNERuntimeRDGUnitTestOpCast.cpp b/Engine/Plugins/Experimental/NNERuntimeRDG/Source/NNERuntimeRDG/Private/Tests/NNERuntimeRDGUnitTestOpCast.cpp
new file mode 100644
--- /dev/null
+++ b/Engine/Plugins/Experimental/NNERuntimeRDG/Source/NNERuntimeRDG/Private/Tests/NNERuntimeRDGUnitTestOpCast.cpp
@@ -0,0 +1,211 @@
+// Copyright Epic Games, Inc. All Rights Reserved.
+
+#include "CoreMinimal.h"
+#include "NNETensor.h"
+#include "NNETypes.h"
+#include "NNERuntimeRDGHelperCast.h"
+#include "NNERuntimeRDGUnitTestHelper.h"
+
+#if WITH_DEV_AUTOMATION_TESTS
+#include "Misc/AutomationTest.h"
+#endif
+
+namespace UE::NNERuntimeRDG::Private::Test::CastOp
+{
+
+#if WITH_DEV_AUTOMATION_TESTS
+
+	using namespace NNE::Internal;
+	using namespace UE::NNERuntimeRDG::Internal;
+
+	IMPLEMENT_NNE_SHAPEINFERENCEHELPER_UNIT_AUTOMATION_TEST(FCastCPUHelperConstOuput, "System.Engine.MachineLearning.NNE.RDG.UnitTest.CastHelper.ConstOutput");
+	bool FCastCPUHelperConstOuput::RunTest(const FString& Parameter)
+	{
+		FTensor XC1 = MakeConstTensor(TEXT("XC1"), { 1 }, { 1.0f });
+		FTensor XC20 = MakeConstTensor(TEXT("XC20"), { 20 }, { 3.0f, 4.0f, 3.0f, 4.0f, 3.0f, 3.0f, 4.0f, 3.0f, 4.0f, 3.0f, 3.0f, 4.0f, 3.0f, 4.0f, 3.0f, 3.0f, 4.0f, 3.0f, 4.0f, 3.0f });
+		FTensor X1 = MakeTensor(TEXT("X"), { 1 });
+
+		FTensor Y = MakeTensor(TEXT("Y"), { 1 });
+		CPUHelper::Cast::Apply(XC1 , Y);
+		UTEST_TRUE(TEXT("Y const if input is const"), Y.HasPreparedData());
+
+		Y = MakeTensor(TEXT("Y"), { 1 });
+		CPUHelper::Cast::Apply(X1, Y);
+		UTEST_FALSE(TEXT("Y not const if input not const"), Y.HasPreparedData());
+
+		Y = MakeTensor(TEXT("Y"), { 20 });
+		CPUHelper::Cast::Apply(XC20, Y);
+		UTEST_FALSE(TEXT("Y not const if input is too large "), Y.HasPreparedData());
+
+		return true;
+	}
+
+	IMPLEMENT_NNE_SHAPEINFERENCEHELPER_UNIT_AUTOMATION_TEST(FCastCPUHelperRank1FromFloat, "System.Engine.MachineLearning.NNE.RDG.UnitTest.CastHelper.Rank1FromFloat");
+	bool FCastCPUHelperRank1FromFloat::RunTest(const FString& Parameter)
+	{
+		FTensor XC2 = MakeConstTensor(TEXT("XC2"), { 2 }, { -1.0f, 1.5f });
+
+		FTensor Y = MakeTensor(TEXT("Y"), { 2 }, ENNETensorDataType::Float);
+		CPUHelper::Cast::Apply(XC2, Y);
+		UTEST_TRUE(TEXT("Y const if input is const"), Y.HasPreparedData());
+		UTEST_EQUAL(TEXT("Cast(float->float)[0]"), Y.GetPreparedData<float>()[0], -1.0f);
+		UTEST_EQUAL(TEXT("Cast(float->float)[1]"), Y.GetPreparedData<float>()[1], 1.5f);
+
+		Y = MakeTensor(TEXT("Y"), { 2 }, ENNETensorDataType::Half);
+		CPUHelper::Cast::Apply(XC2, Y);
+		UTEST_TRUE(TEXT("Y const if input is const"), Y.HasPreparedData());
+		UTEST_EQUAL(TEXT("Cast(float->half)[0]"), Y.GetPreparedData<FFloat16>()[0], -1.0f);
+		UTEST_EQUAL(TEXT("Cast(float->half)[1]"), Y.GetPreparedData<FFloat16>()[1], 1.5f);
+
+		Y = MakeTensor(TEXT("Y"), { 2 }, ENNETensorDataType::Int32);
+		CPUHelper::Cast::Apply(XC2, Y);
+		UTEST_TRUE(TEXT("Y const if input is const"), Y.HasPreparedData());
+		UTEST_EQUAL(TEXT("Cast(float->int32)[0]"), Y.GetPreparedData<int32>()[0], (int32)-1);
+		UTEST_EQUAL(TEXT("Cast(float->int32)[1]"), Y.GetPreparedData<int32>()[1], (int32)1);
+
+		Y = MakeTensor(TEXT("Y"), { 2 }, ENNETensorDataType::Int64);
+		CPUHelper::Cast::Apply(XC2, Y);
+		UTEST_TRUE(TEXT("Y const if input is const"), Y.HasPreparedData());
+		UTEST_EQUAL(TEXT("Cast(float->int64)[0]"), Y.GetPreparedData<int64>()[0], (int64)-1);
+		UTEST_EQUAL(TEXT("Cast(float->int64)[1]"), Y.GetPreparedData<int64>()[1], (int64)1);
+
+		return true;
+	}
+
+	IMPLEMENT_NNE_SHAPEINFERENCEHELPER_UNIT_AUTOMATION_TEST(FCastCPUHelperRank1FromFloat16, "System.Engine.MachineLearning.NNE.RDG.UnitTest.CastHelper.Rank1FromFloat16");
+	bool FCastCPUHelperRank1FromFloat16::RunTest(const FString& Parameter)
+	{
+		FTensor XC2 = MakeConstTensorFloat16(TEXT("XC2"), { 2 }, { -1.0f, 1.5f });
+
+		FTensor Y = MakeTensor(TEXT("Y"), { 2 }, ENNETensorDataType::Float);
+		CPUHelper::Cast::Apply(XC2, Y);
+		UTEST_TRUE(TEXT("Y const if input is const"), Y.HasPreparedData());
+		UTEST_EQUAL(TEXT("Cast(half->float)[0]"), Y.GetPreparedData<float>()[0], -1.0f);
+		UTEST_EQUAL(TEXT("Cast(half->float)[1]"), Y.GetPreparedData<float>()[1], 1.5f);
+
+		Y = MakeTensor(TEXT("Y"), { 2 }, ENNETensorDataType::Half);
+		CPUHelper::Cast::Apply(XC2, Y);
+		UTEST_TRUE(TEXT("Y const if input is const"), Y.HasPreparedData());
+		UTEST_EQUAL(TEXT("Cast(half->half)[0]"), Y.GetPreparedData<FFloat16>()[0], -1.0f);
+		UTEST_EQUAL(TEXT("Cast(half->half)[1]"), Y.GetPreparedData<FFloat16>()[1], 1.5f);
+
+		Y = MakeTensor(TEXT("Y"), { 2 }, ENNETensorDataType::Int32);
+		CPUHelper::Cast::Apply(XC2, Y);
+		UTEST_TRUE(TEXT("Y const if input is const"), Y.HasPreparedData());
+		UTEST_EQUAL(TEXT("Cast(half->int32)[0]"), Y.GetPreparedData<int32>()[0], (int32)-1);
+		UTEST_EQUAL(TEXT("Cast(half->int32)[1]"), Y.GetPreparedData<int32>()[1], (int32)1);
+
+		Y = MakeTensor(TEXT("Y"), { 2 }, ENNETensorDataType::Int64);
+		CPUHelper::Cast::Apply(XC2, Y);
+		UTEST_TRUE(TEXT("Y const if input is const"), Y.HasPreparedData());
+		UTEST_EQUAL(TEXT("Cast(half->int64)[0]"), Y.GetPreparedData<int64>()[0], (int64)-1);
+		UTEST_EQUAL(TEXT("Cast(half->int64)[1]"), Y.GetPreparedData<int64>()[1], (int64)1);
+
+		return true;
+	}
+
+	IMPLEMENT_NNE_SHAPEINFERENCEHELPER_UNIT_AUTOMATION_TEST(FCastCPUHelperRank1FromInt32, "System.Engine.MachineLearning.NNE.RDG.UnitTest.CastHelper.Rank1FromInt32");
+	bool FCastCPUHelperRank1FromInt32::RunTest(const FString& Parameter)
+	{
+		FTensor XC2 = MakeConstTensorInt32(TEXT("XC2"), { 2 }, { -2, 0 });
+
+		FTensor Y = MakeTensor(TEXT("Y"), { 2 }, ENNETensorDataType::Float);
+		CPUHelper::Cast::Apply(XC2, Y);
+		UTEST_TRUE(TEXT("Y const if input is const"), Y.HasPreparedData());
+		UTEST_EQUAL(TEXT("Cast(int32->float)[0]"), Y.GetPreparedData<float>()[0], -2.0f);
+		UTEST_EQUAL(TEXT("Cast(int32->float)[1]"), Y.GetPreparedData<float>()[1], 0.0f);
+
+		Y = MakeTensor(TEXT("Y"), { 2 }, ENNETensorDataType::Half);
+		CPUHelper::Cast::Apply(XC2, Y);
+		UTEST_TRUE(TEXT("Y const if input is const"), Y.HasPreparedData());
+		UTEST_EQUAL(TEXT("Cast(int32->half)[0]"), Y.GetPreparedData<FFloat16>()[0], -2.0f);
+		UTEST_EQUAL(TEXT("Cast(int32->half)[1]"), Y.GetPreparedData<FFloat16>()[1], 0.0f);
+
+		Y = MakeTensor(TEXT("Y"), { 2 }, ENNETensorDataType::Int32);
+		CPUHelper::Cast::Apply(XC2, Y);
+		UTEST_TRUE(TEXT("Y const if input is const"), Y.HasPreparedData());
+		UTEST_EQUAL(TEXT("Cast(int32->int32)[0]"), Y.GetPreparedData<int32>()[0], (int32)-2);
+		UTEST_EQUAL(TEXT("Cast(int32->int32)[1]"), Y.GetPreparedData<int32>()[1], (int32)0);
+
+		Y = MakeTensor(TEXT("Y"), { 2 }, ENNETensorDataType::Int64);
+		CPUHelper::Cast::Apply(XC2, Y);
+		UTEST_TRUE(TEXT("Y const if input is const"), Y.HasPreparedData());
+		UTEST_EQUAL(TEXT("Cast(int32->int64)[0]"), Y.GetPreparedData<int64>()[0], (int64)-2);
+		UTEST_EQUAL(TEXT("Cast(int32->int64)[1]"), Y.GetPreparedData<int64>()[1], (int64)0);
+
+		return true;
+	}
+
+	IMPLEMENT_NNE_SHAPEINFERENCEHELPER_UNIT_AUTOMATION_TEST(FCastCPUHelperRank1FromInt64, "System.Engine.MachineLearning.NNE.RDG.UnitTest.CastHelper.Rank1FromInt64");
+	bool FCastCPUHelperRank1FromInt64::RunTest(const FString& Parameter)
+	{
+		FTensor XC2 = MakeConstTensorInt64(TEXT("XC2"), { 2 }, { (int64)-4, (int64)10});
+
+		FTensor Y = MakeTensor(TEXT("Y"), { 2 }, ENNETensorDataType::Float);
+		CPUHelper::Cast::Apply(XC2, Y);
+		UTEST_TRUE(TEXT("Y const if input is const"), Y.HasPreparedData());
+		UTEST_EQUAL(TEXT("Cast(int64->float)[0]"), Y.GetPreparedData<float>()[0], -4.0f);
+		UTEST_EQUAL(TEXT("Cast(int64->float)[1]"), Y.GetPreparedData<float>()[1], 10.0f);
+
+		Y = MakeTensor(TEXT("Y"), { 2 }, ENNETensorDataType::Half);
+		CPUHelper::Cast::Apply(XC2, Y);
+		UTEST_TRUE(TEXT("Y const if input is const"), Y.HasPreparedData());
+		UTEST_EQUAL(TEXT("Cast(int64->half)[0]"), Y.GetPreparedData<FFloat16>()[0], -4.0f);
+		UTEST_EQUAL(TEXT("Cast(int64->half)[1]"), Y.GetPreparedData<FFloat16>()[1], 10.0f);
+
+		Y = MakeTensor(TEXT("Y"), { 2 }, ENNETensorDataType::Int32);
+		CPUHelper::Cast::Apply(XC2, Y);
+		UTEST_TRUE(TEXT("Y const if input is const"), Y.HasPreparedData());
+		UTEST_EQUAL(TEXT("Cast(int64->int32)[0]"), Y.GetPreparedData<int32>()[0], (int32)-4);
+		UTEST_EQUAL(TEXT("Cast(int64->int32)[1]"), Y.GetPreparedData<int32>()[1], (int32)10);
+
+		Y = MakeTensor(TEXT("Y"), { 2 }, ENNETensorDataType::Int64);
+		CPUHelper::Cast::Apply(XC2, Y);
+		UTEST_TRUE(TEXT("Y const if input is const"), Y.HasPreparedData());
+		UTEST_EQUAL(TEXT("Cast(int64->int64)[0]"), Y.GetPreparedData<int64>()[0], (int64)-4);
+		UTEST_EQUAL(TEXT("Cast(int64->int64)[1]"), Y.GetPreparedData<int64>()[1], (int64)10);
+
+		return true;
+	}
+
+	IMPLEMENT_NNE_SHAPEINFERENCEHELPER_UNIT_AUTOMATION_TEST(FCastCPUHelperRank3, "System.Engine.MachineLearning.NNE.RDG.UnitTest.CastHelper.Rank3");
+	bool FCastCPUHelperRank3::RunTest(const FString& Parameter)
+	{
+		FTensor XC1x2x3 = MakeConstTensor(TEXT("XC1x2x3"), { 1,2,3 }, { -1.0f, -0.5f, 0.0f, 0.1f, 0.6f, 1.5f });
+
+		FTensor Y = MakeTensor(TEXT("Y"), { 1,2,3 }, ENNETensorDataType::Float);
+		CPUHelper::Cast::Apply(XC1x2x3, Y);
+		UTEST_TRUE(TEXT("Y const if input is const"), Y.HasPreparedData());
+		UTEST_EQUAL(TEXT("Cast(float,float)[0]"), Y.GetPreparedData<float>()[0], -1.0f);
+		UTEST_EQUAL(TEXT("Cast(float,float)[1]"), Y.GetPreparedData<float>()[1], -0.5f);
+		UTEST_EQUAL(TEXT("Cast(float,float)[2]"), Y.GetPreparedData<float>()[2],  0.0f);
+		UTEST_EQUAL(TEXT("Cast(float,float)[3]"), Y.GetPreparedData<float>()[3],  0.1f);
+		UTEST_EQUAL(TEXT("Cast(float,float)[4]"), Y.GetPreparedData<float>()[4],  0.6f);
+		UTEST_EQUAL(TEXT("Cast(float,float)[5]"), Y.GetPreparedData<float>()[5],  1.5f);
+
+		Y = MakeTensor(TEXT("Y"), { 1,2,3 }, ENNETensorDataType::Int32);
+		CPUHelper::Cast::Apply(XC1x2x3, Y);
+		UTEST_TRUE(TEXT("Y const if input is const"), Y.HasPreparedData());
+		UTEST_EQUAL(TEXT("Cast(float,int32)[0]"), Y.GetPreparedData<int32>()[0], -1);
+		UTEST_EQUAL(TEXT("Cast(float,int32)[1]"), Y.GetPreparedData<int32>()[1], 0);
+		UTEST_EQUAL(TEXT("Cast(float,int32)[2]"), Y.GetPreparedData<int32>()[2], 0);
+		UTEST_EQUAL(TEXT("Cast(float,int32)[3]"), Y.GetPreparedData<int32>()[3], 0);
+		UTEST_EQUAL(TEXT("Cast(float,int32)[4]"), Y.GetPreparedData<int32>()[4], 0);
+		UTEST_EQUAL(TEXT("Cast(float,int32)[5]"), Y.GetPreparedData<int32>()[5], 1);
+
+		Y = MakeTensor(TEXT("Y"), { 1,2,3 }, ENNETensorDataType::Int64);
+		CPUHelper::Cast::Apply(XC1x2x3, Y);
+		UTEST_TRUE(TEXT("Y const if input is const"), Y.HasPreparedData());
+		UTEST_EQUAL(TEXT("Cast(float,int64)[0]"), Y.GetPreparedData<int64>()[0], (int64)-1);
+		UTEST_EQUAL(TEXT("Cast(float,int64)[1]"), Y.GetPreparedData<int64>()[1], (int64)0);
+		UTEST_EQUAL(TEXT("Cast(float,int64)[2]"), Y.GetPreparedData<int64>()[2], (int64)0);
+		UTEST_EQUAL(TEXT("Cast(float,int64)[3]"), Y.GetPreparedData<int64>()[3], (int64)0);
+		UTEST_EQUAL(TEXT("Cast(float,int64)[4]"), Y.GetPreparedData<int64>()[4], (int64)0);
+		UTEST_EQUAL(TEXT("Cast(float,int64)[5]"), Y.GetPreparedData<int64>()[5], (int64)1);
+
+		return true;
+	}
+
+#endif //WITH_DEV_AUTOMATION_TESTS
+
+} // namespace UE::NNERuntimeRDG::Private::Test::CastOp
diff --git a/Engine/Plugins/Experimental/NNERuntimeRDG/Source/NNERuntimeRDG/Private/Tests/NNERuntimeRDGUnitTestOpConcat.cpp b/Engine/Plugins/Experimental/NNERuntimeRDG/Source/NNERuntimeRDG/Private/Tests/NNERuntimeRDGUnitTestOpConcat.cpp
new file mode 100644
--- /dev/null
+++ b/Engine/Plugins/Experimental/NNERuntimeRDG/Source/NNERuntimeRDG/Private/Tests/NNERuntimeRDGUnitTestOpConcat.cpp
@@ -0,0 +1,127 @@
+// Copyright Epic Games, Inc. All Rights Reserved.
+
+#include "CoreMinimal.h"
+#include "NNETensor.h"
+#include "NNETypes.h"
+#include "NNERuntimeRDGHelperConcat.h"
+#include "NNERuntimeRDGUnitTestHelper.h"
+
+#if WITH_DEV_AUTOMATION_TESTS
+#include "Misc/AutomationTest.h"
+#endif
+
+namespace UE::NNERuntimeRDG::Private::Test::ConcatOp
+{
+
+#if WITH_DEV_AUTOMATION_TESTS
+
+	using namespace NNE::Internal;
+	using namespace UE::NNERuntimeRDG::Internal;
+
+	IMPLEMENT_NNE_SHAPEINFERENCEHELPER_UNIT_AUTOMATION_TEST(FConcatCPUHelperConstOuput, "System.Engine.MachineLearning.NNE.RDG.UnitTest.ConcatHelper.ConstOutput");
+	bool FConcatCPUHelperConstOuput::RunTest(const FString& Parameter)
+	{
+		FTensor XC1 = MakeConstTensor(TEXT("XC1"), { 1 }, { 1.0f });
+		FTensor XC20 = MakeConstTensor(TEXT("XC20"), { 20 }, { 3.0f, 4.0f, 3.0f, 4.0f, 3.0f, 3.0f, 4.0f, 3.0f, 4.0f, 3.0f, 3.0f, 4.0f, 3.0f, 4.0f, 3.0f, 3.0f, 4.0f, 3.0f, 4.0f, 3.0f });
+		FTensor X1 = MakeTensor(TEXT("X"), { 1 });
+
+		FTensor Y = MakeTensor(TEXT("Y"), { 1 });
+		CPUHelper::Concat::Apply({ &XC1 }, Y, 0);
+		UTEST_TRUE(TEXT("Y const if input is const"), Y.HasPreparedData());
+
+		Y = MakeTensor(TEXT("Y"), { 2 });
+		CPUHelper::Concat::Apply({ &XC1, &XC1 }, Y, 0);
+		UTEST_TRUE(TEXT("Y const if inputs are const"), Y.HasPreparedData());
+
+		Y = MakeTensor(TEXT("Y"), { 2 });
+		CPUHelper::Concat::Apply({ &X1, &XC1 }, Y, 0);
+		UTEST_FALSE(TEXT("Y not const if at least on input not const"), Y.HasPreparedData());
+
+		Y = MakeTensor(TEXT("Y"), { 20 });
+		CPUHelper::Concat::Apply({ &XC20 }, Y, 0);
+		UTEST_FALSE(TEXT("Y not const if input is too large "), Y.HasPreparedData());
+
+		return true;
+	}
+
+	IMPLEMENT_NNE_SHAPEINFERENCEHELPER_UNIT_AUTOMATION_TEST(FConcatCPUHelperRank1, "System.Engine.MachineLearning.NNE.RDG.UnitTest.ConcatHelper.Rank1");
+	bool FConcatCPUHelperRank1::RunTest(const FString& Parameter)
+	{
+		FTensor XC1 = MakeConstTensor(TEXT("XC1"), { 1 }, { 1.0f });
+		FTensor XC1v2 = MakeConstTensor(TEXT("XC1"), { 1 }, { 2.0f });
+
+		FTensor Y = MakeTensor(TEXT("Y"), { 1 });
+		CPUHelper::Concat::Apply({ &XC1 }, Y, 0);
+		UTEST_TRUE(TEXT("Y const if input is const"), Y.HasPreparedData());
+		UTEST_EQUAL(TEXT("Concat(XC1,XC1)[0]"), Y.GetPreparedData<float>()[0], 1.0f);
+
+		Y = MakeTensor(TEXT("Y"), { 2 });
+		CPUHelper::Concat::Apply({ &XC1, &XC1v2 }, Y, 0);
+		UTEST_TRUE(TEXT("Y const if inputs are const"), Y.HasPreparedData());
+		UTEST_EQUAL(TEXT("Concat(XC1,XC1v2)[0]"), Y.GetPreparedData<float>()[0], 1.0f);
+		UTEST_EQUAL(TEXT("Concat(XC1,XC1v2)[1]"), Y.GetPreparedData<float>()[1], 2.0f);
+
+		Y = MakeTensor(TEXT("Y"), { 3 });
+		CPUHelper::Concat::Apply({ &XC1, &XC1v2, &XC1 }, Y, 0);
+		UTEST_TRUE(TEXT("Y const if inputs are const"), Y.HasPreparedData());
+		UTEST_EQUAL(TEXT("Concat(XC1,XC1v2,XC1)[0]"), Y.GetPreparedData<float>()[0], 1.0f);
+		UTEST_EQUAL(TEXT("Concat(XC1,XC1v2,XC1)[1]"), Y.GetPreparedData<float>()[1], 2.0f);
+		UTEST_EQUAL(TEXT("Concat(XC1,XC1v2,XC1)[2]"), Y.GetPreparedData<float>()[2], 1.0f);
+
+		return true;
+	}
+
+	IMPLEMENT_NNE_SHAPEINFERENCEHELPER_UNIT_AUTOMATION_TEST(FConcatCPUHelperRank3, "System.Engine.MachineLearning.NNE.RDG.UnitTest.ConcatHelper.Rank3");
+	bool FConcatCPUHelperRank3::RunTest(const FString& Parameter)
+	{
+		FTensor XC1x2x3 = MakeConstTensor(TEXT("XC1x2x3"), { 1,2,3 }, { 1.0f, 2.0f, 3.0f, 4.0f, 5.0f, 6.0f });
+		FTensor XC1x2x1 = MakeConstTensor(TEXT("XC1x2x1"), { 1,2,1 }, { -1.0f, -2.0f });
+		FTensor XC1x1x3 = MakeConstTensor(TEXT("XC1x1x3"), { 1,1,3 }, { 10.0f, 11.0f, 12.0f });
+
+		FTensor Y = MakeTensor(TEXT("Y"), { 1,2,4 });
+		CPUHelper::Concat::Apply({ &XC1x2x3, &XC1x2x1 }, Y, 2);
+		UTEST_TRUE(TEXT("Y const if inputs are const"), Y.HasPreparedData());
+		UTEST_EQUAL(TEXT("Concat(XC1x2x3,XC1x2x1)[0]"), Y.GetPreparedData<float>()[0], 1.0f);
+		UTEST_EQUAL(TEXT("Concat(XC1x2x3,XC1x2x1)[1]"), Y.GetPreparedData<float>()[1], 2.0f);
+		UTEST_EQUAL(TEXT("Concat(XC1x2x3,XC1x2x1)[2]"), Y.GetPreparedData<float>()[2], 3.0f);
+		UTEST_EQUAL(TEXT("Concat(XC1x2x3,XC1x2x1)[3]"), Y.GetPreparedData<float>()[3], -1.0f);
+		UTEST_EQUAL(TEXT("Concat(XC1x2x3,XC1x2x1)[4]"), Y.GetPreparedData<float>()[4], 4.0f);
+		UTEST_EQUAL(TEXT("Concat(XC1x2x3,XC1x2x1)[5]"), Y.GetPreparedData<float>()[5], 5.0f);
+		UTEST_EQUAL(TEXT("Concat(XC1x2x3,XC1x2x1)[6]"), Y.GetPreparedData<float>()[6], 6.0f);
+		UTEST_EQUAL(TEXT("Concat(XC1x2x3,XC1x2x1)[7]"), Y.GetPreparedData<float>()[7], -2.0f);
+
+		Y = MakeTensor(TEXT("Y"), { 1,3,3 });
+		CPUHelper::Concat::Apply({ &XC1x2x3, &XC1x1x3 }, Y, 1);
+		UTEST_TRUE(TEXT("Y const if inputs are const"), Y.HasPreparedData());
+		UTEST_EQUAL(TEXT("Concat(XC1x2x3,XC1x1x3)[0]"), Y.GetPreparedData<float>()[0], 1.0f);
+		UTEST_EQUAL(TEXT("Concat(XC1x2x3,XC1x1x3)[1]"), Y.GetPreparedData<float>()[1], 2.0f);
+		UTEST_EQUAL(TEXT("Concat(XC1x2x3,XC1x1x3)[2]"), Y.GetPreparedData<float>()[2], 3.0f);
+		UTEST_EQUAL(TEXT("Concat(XC1x2x3,XC1x1x3)[3]"), Y.GetPreparedData<float>()[3], 4.0f);
+		UTEST_EQUAL(TEXT("Concat(XC1x2x3,XC1x1x3)[4]"), Y.GetPreparedData<float>()[4], 5.0f);
+		UTEST_EQUAL(TEXT("Concat(XC1x2x3,XC1x1x3)[5]"), Y.GetPreparedData<float>()[5], 6.0f);
+		UTEST_EQUAL(TEXT("Concat(XC1x2x3,XC1x1x3)[6]"), Y.GetPreparedData<float>()[6], 10.0f);
+		UTEST_EQUAL(TEXT("Concat(XC1x2x3,XC1x1x3)[7]"), Y.GetPreparedData<float>()[7], 11.0f);
+		UTEST_EQUAL(TEXT("Concat(XC1x2x3,XC1x1x3)[8]"), Y.GetPreparedData<float>()[8], 12.0f);
+
+		Y = MakeTensor(TEXT("Y"), { 2,2,3 });
+		CPUHelper::Concat::Apply({ &XC1x2x3, &XC1x2x3 }, Y, 0);
+		UTEST_TRUE(TEXT("Y const if inputs are const"), Y.HasPreparedData());
+		UTEST_EQUAL(TEXT("Concat(XC1x2x3,XC1x2x3)[0]"), Y.GetPreparedData<float>()[0], 1.0f);
+		UTEST_EQUAL(TEXT("Concat(XC1x2x3,XC1x2x3)[1]"), Y.GetPreparedData<float>()[1], 2.0f);
+		UTEST_EQUAL(TEXT("Concat(XC1x2x3,XC1x2x3)[2]"), Y.GetPreparedData<float>()[2], 3.0f);
+		UTEST_EQUAL(TEXT("Concat(XC1x2x3,XC1x2x3)[3]"), Y.GetPreparedData<float>()[3], 4.0f);
+		UTEST_EQUAL(TEXT("Concat(XC1x2x3,XC1x2x3)[4]"), Y.GetPreparedData<float>()[4], 5.0f);
+		UTEST_EQUAL(TEXT("Concat(XC1x2x3,XC1x2x3)[5]"), Y.GetPreparedData<float>()[5], 6.0f);
+		UTEST_EQUAL(TEXT("Concat(XC1x2x3,XC1x2x3)[6]"), Y.GetPreparedData<float>()[6], 1.0f);
+		UTEST_EQUAL(TEXT("Concat(XC1x2x3,XC1x2x3)[7]"), Y.GetPreparedData<float>()[7], 2.0f);
+		UTEST_EQUAL(TEXT("Concat(XC1x2x3,XC1x2x3)[8]"), Y.GetPreparedData<float>()[8], 3.0f);
+		UTEST_EQUAL(TEXT("Concat(XC1x2x3,XC1x2x3)[9]"), Y.GetPreparedData<float>()[9], 4.0f);
+		UTEST_EQUAL(TEXT("Concat(XC1x2x3,XC1x2x3)[10]"), Y.GetPreparedData<float>()[10], 5.0f);
+		UTEST_EQUAL(TEXT("Concat(XC1x2x3,XC1x2x3)[11]"), Y.GetPreparedData<float>()[11], 6.0f);
+
+		return true;
+	}
+
+#endif //WITH_DEV_AUTOMATION_TESTS
+
+} // namespace UE::NNERuntimeRDG::Private::Test::ConcatOp
diff --git a/Engine/Plugins/Experimental/NNERuntimeRDG/Source/NNERuntimeRDG/Private/Tests/NNERuntimeRDGUnitTestOpGather.cpp b/Engine/Plugins/Experimental/NNERuntimeRDG/Source/NNERuntimeRDG/Private/Tests/NNERuntimeRDGUnitTestOpGather.cpp
new file mode 100644
--- /dev/null
+++ b/Engine/Plugins/Experimental/NNERuntimeRDG/Source/NNERuntimeRDG/Private/Tests/NNERuntimeRDGUnitTestOpGather.cpp
@@ -0,0 +1,171 @@
+// Copyright Epic Games, Inc. All Rights Reserved.
+
+#include "CoreMinimal.h"
+#include "NNETensor.h"
+#include "NNETypes.h"
+#include "NNERuntimeRDGHelperGather.h"
+#include "NNERuntimeRDGUnitTestHelper.h"
+
+#if WITH_DEV_AUTOMATION_TESTS
+#include "Misc/AutomationTest.h"
+#endif
+
+namespace UE::NNERuntimeRDG::Private::Test::GatherOp
+{
+
+#if WITH_DEV_AUTOMATION_TESTS
+
+	using namespace NNE::Internal;
+	using namespace UE::NNERuntimeRDG::Internal;
+
+	IMPLEMENT_NNE_SHAPEINFERENCEHELPER_UNIT_AUTOMATION_TEST(FGatherCPUHelperConstOuput, "System.Engine.MachineLearning.NNE.RDG.UnitTest.GatherHelper.ConstOutput");
+	bool FGatherCPUHelperConstOuput::RunTest(const FString& Parameter)
+	{
+		FTensor XC1 = MakeConstTensorInt64(TEXT("XC1"), { 1 }, { 0 });
+		FTensor XC20 = MakeConstTensorInt64(TEXT("XC20"), { 20 }, { 0,1,2,3,4,5,6,7,8,9,0,1,2,3,4,5,6,7,8,9 });
+		FTensor X1 = MakeTensor(TEXT("X"), { 0 }, ENNETensorDataType::Int64);
+
+		FTensor Y = MakeTensor(TEXT("Y"), { 1 }, ENNETensorDataType::Int64);
+		CPUHelper::Gather::Apply(XC1 , XC1, 0, Y);
+		UTEST_TRUE(TEXT("Y const if inputs are const"), Y.HasPreparedData());
+
+		Y = MakeTensor(TEXT("Y"), { 1 }, ENNETensorDataType::Int64);
+		CPUHelper::Gather::Apply(XC1, X1, 0, Y);
+		UTEST_FALSE(TEXT("Y not const if data not const"), Y.HasPreparedData());
+
+		Y = MakeTensor(TEXT("Y"), { 1 }, ENNETensorDataType::Int64);
+		CPUHelper::Gather::Apply(X1, XC1, 0, Y);
+		UTEST_FALSE(TEXT("Y not const if indices not const"), Y.HasPreparedData());
+
+		Y = MakeTensor(TEXT("Y"), { 20 }, ENNETensorDataType::Int64);
+		CPUHelper::Gather::Apply(XC20, XC1, 0, Y);
+		UTEST_FALSE(TEXT("Y not const if input is too large "), Y.HasPreparedData());
+
+		return true;
+	}
+
+	IMPLEMENT_NNE_SHAPEINFERENCEHELPER_UNIT_AUTOMATION_TEST(FGatherCPUHelperRank1Int64Indices, "System.Engine.MachineLearning.NNE.RDG.UnitTest.GatherHelper.Rank1Int64Indices");
+	bool FGatherCPUHelperRank1Int64Indices::RunTest(const FString& Parameter)
+	{
+		FTensor XC4i64 = MakeConstTensorInt64(TEXT("XC4"), { 4 }, {    0,    2,    -3,     5 });
+		FTensor XC4i32 = MakeConstTensorInt32(TEXT("XC4"), { 4 }, {   -1,    3,    13,     8 });
+		FTensor XC4f32 = MakeConstTensor(TEXT("XC4"),      { 4 }, { 0.1f, 2.5f, -3.1f, 15.0f });
+
+		FTensor XC3 = MakeConstTensorInt64(TEXT("XC3"), { 3 }, { 0, 2 ,3 });
+		FTensor XC2 = MakeConstTensorInt64(TEXT("XC2"), { 2 }, { 1, 2 });
+
+		//int64 Data
+		{
+			FTensor Y = MakeTensor(TEXT("Y"), { 3 }, ENNETensorDataType::Int64);
+			CPUHelper::Gather::Apply(XC4i64, XC3, 0, Y);
+			UTEST_TRUE(TEXT("Y const if input is const"), Y.HasPreparedData());
+			UTEST_EQUAL(TEXT("Gather(XC4i64,XC3)[0]"), Y.GetPreparedData<int64>()[0], (int64)0);
+			UTEST_EQUAL(TEXT("Gather(XC4i64,XC3)[1]"), Y.GetPreparedData<int64>()[1], (int64)-3);
+			UTEST_EQUAL(TEXT("Gather(XC4i64,XC3)[2]"), Y.GetPreparedData<int64>()[2], (int64)5);
+
+			Y = MakeTensor(TEXT("Y"), { 2 }, ENNETensorDataType::Int64);
+			CPUHelper::Gather::Apply(XC4i64, XC2, 0, Y);
+			UTEST_TRUE(TEXT("Y const if input is const"), Y.HasPreparedData());
+			UTEST_EQUAL(TEXT("Gather(XC4i64,XC2)[0]"), Y.GetPreparedData<int64>()[0], (int64)2);
+			UTEST_EQUAL(TEXT("Gather(XC4i64,XC2)[1]"), Y.GetPreparedData<int64>()[1], (int64)-3);
+		}
+
+		//int32 Data
+		{
+			FTensor Y = MakeTensor(TEXT("Y"), { 3 }, ENNETensorDataType::Int32);
+			CPUHelper::Gather::Apply(XC4i32, XC3, 0, Y);
+			UTEST_TRUE(TEXT("Y const if input is const"), Y.HasPreparedData());
+			UTEST_EQUAL(TEXT("Gather(XC4i32,XC3)[0]"), Y.GetPreparedData<int32>()[0], -1);
+			UTEST_EQUAL(TEXT("Gather(XC4i32,XC3)[1]"), Y.GetPreparedData<int32>()[1], 13);
+			UTEST_EQUAL(TEXT("Gather(XC4i32,XC3)[2]"), Y.GetPreparedData<int32>()[2], 8);
+
+			Y = MakeTensor(TEXT("Y"), { 2 }, ENNETensorDataType::Int32);
+			CPUHelper::Gather::Apply(XC4i32, XC2, 0, Y);
+			UTEST_TRUE(TEXT("Y const if input is const"), Y.HasPreparedData());
+			UTEST_EQUAL(TEXT("Gather(XC4i32,XC2)[0]"), Y.GetPreparedData<int32>()[0], 3);
+			UTEST_EQUAL(TEXT("Gather(XC4i32,XC2)[1]"), Y.GetPreparedData<int32>()[1], 13);
+		}
+
+		//float Data
+		{
+			FTensor Y = MakeTensor(TEXT("Y"), { 3 }, ENNETensorDataType::Float);
+			CPUHelper::Gather::Apply(XC4f32, XC3, 0, Y);
+			UTEST_TRUE(TEXT("Y const if input is const"), Y.HasPreparedData());
+			UTEST_EQUAL(TEXT("Gather(XC4f32,XC3)[0]"), Y.GetPreparedData<float>()[0], 0.1f);
+			UTEST_EQUAL(TEXT("Gather(XC4f32,XC3)[1]"), Y.GetPreparedData<float>()[1], -3.1f);
+			UTEST_EQUAL(TEXT("Gather(XC4f32,XC3)[2]"), Y.GetPreparedData<float>()[2], 15.0f);
+
+			Y = MakeTensor(TEXT("Y"), { 2 }, ENNETensorDataType::Float);
+			CPUHelper::Gather::Apply(XC4f32, XC2, 0, Y);
+			UTEST_TRUE(TEXT("Y const if input is const"), Y.HasPreparedData());
+			UTEST_EQUAL(TEXT("Gather(XC4f32,XC2)[0]"), Y.GetPreparedData<float>()[0], 2.5f);
+			UTEST_EQUAL(TEXT("Gather(XC4f32,XC2)[1]"), Y.GetPreparedData<float>()[1], -3.1f);
+		}
+
+		return true;
+	}
+
+	IMPLEMENT_NNE_SHAPEINFERENCEHELPER_UNIT_AUTOMATION_TEST(FGatherCPUHelperRank1Int32Indices, "System.Engine.MachineLearning.NNE.RDG.UnitTest.GatherHelper.Rank1Int32Indices");
+	bool FGatherCPUHelperRank1Int32Indices::RunTest(const FString& Parameter)
+	{
+		FTensor XC4i64 = MakeConstTensorInt64(TEXT("XC4"), { 4 }, {    0,    2,    -3,     5 });
+		FTensor XC4i32 = MakeConstTensorInt32(TEXT("XC4"), { 4 }, {   -1,    3,    13,     8 });
+		FTensor XC4f32 = MakeConstTensor(TEXT("XC4"),      { 4 }, { 0.1f, 2.5f, -3.1f, 15.0f });
+
+		FTensor XC3 = MakeConstTensorInt32(TEXT("XC3"), { 3 }, { 0, 2 ,3 });
+		FTensor XC2 = MakeConstTensorInt32(TEXT("XC2"), { 2 }, { 1, 2 });
+
+		//int64 Data
+		{
+			FTensor Y = MakeTensor(TEXT("Y"), { 3 }, ENNETensorDataType::Int64);
+			CPUHelper::Gather::Apply(XC4i64, XC3, 0, Y);
+			UTEST_TRUE(TEXT("Y const if input is const"), Y.HasPreparedData());
+			UTEST_EQUAL(TEXT("Gather(XC4i64,XC3)[0]"), Y.GetPreparedData<int64>()[0], (int64)0);
+			UTEST_EQUAL(TEXT("Gather(XC4i64,XC3)[1]"), Y.GetPreparedData<int64>()[1], (int64)-3);
+			UTEST_EQUAL(TEXT("Gather(XC4i64,XC3)[2]"), Y.GetPreparedData<int64>()[2], (int64)5);
+
+			Y = MakeTensor(TEXT("Y"), { 2 }, ENNETensorDataType::Int64);
+			CPUHelper::Gather::Apply(XC4i64, XC2, 0, Y);
+			UTEST_TRUE(TEXT("Y const if input is const"), Y.HasPreparedData());
+			UTEST_EQUAL(TEXT("Gather(XC4i64,XC2)[0]"), Y.GetPreparedData<int64>()[0], (int64)2);
+			UTEST_EQUAL(TEXT("Gather(XC4i64,XC2)[1]"), Y.GetPreparedData<int64>()[1], (int64)-3);
+		}
+
+		//int32 Data
+		{
+			FTensor Y = MakeTensor(TEXT("Y"), { 3 }, ENNETensorDataType::Int32);
+			CPUHelper::Gather::Apply(XC4i32, XC3, 0, Y);
+			UTEST_TRUE(TEXT("Y const if input is const"), Y.HasPreparedData());
+			UTEST_EQUAL(TEXT("Gather(XC4i32,XC3)[0]"), Y.GetPreparedData<int32>()[0], -1);
+			UTEST_EQUAL(TEXT("Gather(XC4i32,XC3)[1]"), Y.GetPreparedData<int32>()[1], 13);
+			UTEST_EQUAL(TEXT("Gather(XC4i32,XC3)[2]"), Y.GetPreparedData<int32>()[2], 8);
+
+			Y = MakeTensor(TEXT("Y"), { 2 }, ENNETensorDataType::Int32);
+			CPUHelper::Gather::Apply(XC4i32, XC2, 0, Y);
+			UTEST_TRUE(TEXT("Y const if input is const"), Y.HasPreparedData());
+			UTEST_EQUAL(TEXT("Gather(XC4i32,XC2)[0]"), Y.GetPreparedData<int32>()[0], 3);
+			UTEST_EQUAL(TEXT("Gather(XC4i32,XC2)[1]"), Y.GetPreparedData<int32>()[1], 13);
+		}
+
+		//float Data
+		{
+			FTensor Y = MakeTensor(TEXT("Y"), { 3 }, ENNETensorDataType::Float);
+			CPUHelper::Gather::Apply(XC4f32, XC3, 0, Y);
+			UTEST_TRUE(TEXT("Y const if input is const"), Y.HasPreparedData());
+			UTEST_EQUAL(TEXT("Gather(XC4f32,XC3)[0]"), Y.GetPreparedData<float>()[0], 0.1f);
+			UTEST_EQUAL(TEXT("Gather(XC4f32,XC3)[1]"), Y.GetPreparedData<float>()[1], -3.1f);
+			UTEST_EQUAL(TEXT("Gather(XC4f32,XC3)[2]"), Y.GetPreparedData<float>()[2], 15.0f);
+
+			Y = MakeTensor(TEXT("Y"), { 2 }, ENNETensorDataType::Float);
+			CPUHelper::Gather::Apply(XC4f32, XC2, 0, Y);
+			UTEST_TRUE(TEXT("Y const if input is const"), Y.HasPreparedData());
+			UTEST_EQUAL(TEXT("Gather(XC4f32,XC2)[0]"), Y.GetPreparedData<float>()[0], 2.5f);
+			UTEST_EQUAL(TEXT("Gather(XC4f32,XC2)[1]"), Y.GetPreparedData<float>()[1], -3.1f);
+		}
+
+		return true;
+	}
+
+#endif //WITH_DEV_AUTOMATION_TESTS
+
+} // namespace UE::NNERuntimeRDG::Private::Test::GatherOp
diff --git a/Engine/Plugins/Experimental/NNERuntimeRDG/Source/NNERuntimeRDG/Private/Tests/NNERuntimeRDGUnitTestOpSlice.cpp b/Engine/Plugins/Experimental/NNERuntimeRDG/Source/NNERuntimeRDG/Private/Tests/NNERuntimeRDGUnitTestOpSlice.cpp
new file mode 100644
--- /dev/null
+++ b/Engine/Plugins/Experimental/NNERuntimeRDG/Source/NNERuntimeRDG/Private/Tests/NNERuntimeRDGUnitTestOpSlice.cpp
@@ -0,0 +1,162 @@
+// Copyright Epic Games, Inc. All Rights Reserved.
+
+#include "CoreMinimal.h"
+#include "NNETensor.h"
+#include "NNETypes.h"
+#include "NNERuntimeRDGHelperSlice.h"
+#include "NNERuntimeRDGUnitTestHelper.h"
+
+#if WITH_DEV_AUTOMATION_TESTS
+#include "Misc/AutomationTest.h"
+#endif
+
+namespace UE::NNERuntimeRDG::Private::Test::SliceOp
+{
+
+#if WITH_DEV_AUTOMATION_TESTS
+
+	using namespace NNE::Internal;
+	
+
+	namespace
+	{
+		void CallSliceApply(const FTensor& InputTensor, FTensor& OutputTensor, TConstArrayView<int32> Starts, TConstArrayView<int32> Steps)
+		{
+			using namespace UE::NNERuntimeRDG::Internal;
+			CPUHelper::Slice::Apply(InputTensor, OutputTensor, Starts, Steps);
+		}
+
+		// Default steps (all 1s)
+		void CallSliceApply(const FTensor& InputTensor, FTensor& OutputTensor, TConstArrayView<int32> Starts)
+		{
+			TArray<int32, TInlineAllocator<NNE::FTensorShape::MaxRank>> Steps;
+			Steps.Init(1, InputTensor.GetShape().Rank());
+			CallSliceApply(InputTensor, OutputTensor, Starts, Steps);
+		}
+	}
+
+	IMPLEMENT_NNE_SHAPEINFERENCEHELPER_UNIT_AUTOMATION_TEST(FSliceCPUHelperConstOuput, "System.Engine.MachineLearning.NNE.RDG.UnitTest.SliceHelper.ConstOutput");
+	bool FSliceCPUHelperConstOuput::RunTest(const FString& Parameter)
+	{
+		FTensor XC1 = MakeConstTensor(TEXT("XC1"), { 1 }, { 1.0f });
+		FTensor XC20 = MakeConstTensor(TEXT("XC20"), { 20 }, { 3.0f, 4.0f, 3.0f, 4.0f, 3.0f, 3.0f, 4.0f, 3.0f, 4.0f, 3.0f, 3.0f, 4.0f, 3.0f, 4.0f, 3.0f, 3.0f, 4.0f, 3.0f, 4.0f, 3.0f });
+		FTensor X1 = MakeTensor(TEXT("X"), { 1 });
+
+		FTensor Y = MakeTensor(TEXT("Y"), { 1 });
+		CallSliceApply(XC1 , Y, { 0 });
+		UTEST_TRUE(TEXT("Y const if input is const"), Y.HasPreparedData());
+
+		Y = MakeTensor(TEXT("Y"), { 1 });
+		CallSliceApply(X1, Y, { 0 });
+		UTEST_FALSE(TEXT("Y not const if input not const"), Y.HasPreparedData());
+
+		Y = MakeTensor(TEXT("Y"), { 20 });
+		CallSliceApply(XC20, Y, { 0 });
+		UTEST_FALSE(TEXT("Y not const if input is too large "), Y.HasPreparedData());
+
+		return true;
+	}
+
+	IMPLEMENT_NNE_SHAPEINFERENCEHELPER_UNIT_AUTOMATION_TEST(FSliceCPUHelperRank1, "System.Engine.MachineLearning.NNE.RDG.UnitTest.SliceHelper.Rank1");
+	bool FSliceCPUHelperRank1::RunTest(const FString& Parameter)
+	{
+		FTensor XC6 = MakeConstTensor(TEXT("XC6"), { 6 }, { 1.0f, 2.0f, 3.0f, 4.0f, 5.0f, 6.0f });
+
+		FTensor Y = MakeTensor(TEXT("Y"), { 6 });
+		CallSliceApply(XC6, Y, { 0 });
+		UTEST_TRUE(TEXT("Y const if input is const"), Y.HasPreparedData());
+		UTEST_EQUAL(TEXT("Slice(XC6,6,0)[0]"), Y.GetPreparedData<float>()[0], 1.0f);
+		UTEST_EQUAL(TEXT("Slice(XC6,6,0)[1]"), Y.GetPreparedData<float>()[1], 2.0f);
+		UTEST_EQUAL(TEXT("Slice(XC6,6,0)[2]"), Y.GetPreparedData<float>()[2], 3.0f);
+		UTEST_EQUAL(TEXT("Slice(XC6,6,0)[3]"), Y.GetPreparedData<float>()[3], 4.0f);
+		UTEST_EQUAL(TEXT("Slice(XC6,6,0)[4]"), Y.GetPreparedData<float>()[4], 5.0f);
+		UTEST_EQUAL(TEXT("Slice(XC6,6,0)[5]"), Y.GetPreparedData<float>()[5], 6.0f);
+
+		Y = MakeTensor(TEXT("Y"), { 2 });
+		CallSliceApply(XC6, Y, { 4 });
+		UTEST_TRUE(TEXT("Y const if input is const"), Y.HasPreparedData());
+		UTEST_EQUAL(TEXT("Slice(XC6,2,4)[0]"), Y.GetPreparedData<float>()[0], 5.0f);
+		UTEST_EQUAL(TEXT("Slice(XC6,2,4)[1]"), Y.GetPreparedData<float>()[1], 6.0f);
+
+		Y = MakeTensor(TEXT("Y"), { 2 });
+		CallSliceApply(XC6, Y, { 1 });
+		UTEST_TRUE(TEXT("Y const if input is const"), Y.HasPreparedData());
+		UTEST_EQUAL(TEXT("Slice(XC6,2,1)[0]"), Y.GetPreparedData<float>()[0], 2.0f);
+		UTEST_EQUAL(TEXT("Slice(XC6,2,1)[1]"), Y.GetPreparedData<float>()[1], 3.0f);
+
+		return true;
+	}
+
+	IMPLEMENT_NNE_SHAPEINFERENCEHELPER_UNIT_AUTOMATION_TEST(FSliceCPUHelperRank1Int, "System.Engine.MachineLearning.NNE.RDG.UnitTest.SliceHelper.Rank1Int");
+	bool FSliceCPUHelperRank1Int::RunTest(const FString& Parameter)
+	{
+		FTensor XC6Int32 = MakeConstTensorInt32(TEXT("XC6Int32"), { 6 }, { 1, 2, 3, 4, 5, 6 });
+		FTensor XC6Int64 = MakeConstTensorInt64(TEXT("XC6Int64"), { 6 }, { 1, 2, 3, 4, 5, 6 });
+		FTensor YInt32 = MakeTensor(TEXT("YInt32"), { 2 }, ENNETensorDataType::Int32);
+		FTensor YInt64 = MakeTensor(TEXT("YInt64"), { 2 }, ENNETensorDataType::Int64);
+
+		CallSliceApply(XC6Int32, YInt32, { 1 });
+		UTEST_TRUE(TEXT("YInt32 const if input is const"), YInt32.HasPreparedData());
+		UTEST_EQUAL(TEXT("Slice(XC6Int32,2,4)[0]"), YInt32.GetPreparedData<int32>()[0], 2);
+		UTEST_EQUAL(TEXT("Slice(XC6Int32,2,4)[1]"), YInt32.GetPreparedData<int32>()[1], 3);
+
+		CallSliceApply(XC6Int64, YInt64, { 1 });
+		UTEST_TRUE(TEXT("YInt64 const if input is const"), YInt64.HasPreparedData());
+		UTEST_EQUAL(TEXT("Slice(XC6Int64,2,4)[0]"), YInt64.GetPreparedData<int64>()[0], (int64)2);
+		UTEST_EQUAL(TEXT("Slice(XC6Int64,2,4)[1]"), YInt64.GetPreparedData<int64>()[1], (int64)3);
+
+		return true;
+	}
+
+	IMPLEMENT_NNE_SHAPEINFERENCEHELPER_UNIT_AUTOMATION_TEST(FSliceCPUHelperRank3, "System.Engine.MachineLearning.NNE.RDG.UnitTest.SliceHelper.Rank3");
+	bool FSliceCPUHelperRank3::RunTest(const FString& Parameter)
+	{
+		FTensor XC1x2x3 = MakeConstTensor(TEXT("XC1x2x3"), { 1,2,3 }, { 1.0f, 2.0f, 3.0f, 4.0f, 5.0f, 6.0f });
+
+		FTensor Y = MakeTensor(TEXT("Y"), { 1,2,3 });
+		CallSliceApply(XC1x2x3, Y, { 0,0,0 });
+		UTEST_TRUE(TEXT("Y const if input is const"), Y.HasPreparedData());
+		UTEST_EQUAL(TEXT("Slice(XC1x2x3,1x2x3,0-0-0)[0]"), Y.GetPreparedData<float>()[0], 1.0f);
+		UTEST_EQUAL(TEXT("Slice(XC1x2x3,1x2x3,0-0-0)[1]"), Y.GetPreparedData<float>()[1], 2.0f);
+		UTEST_EQUAL(TEXT("Slice(XC1x2x3,1x2x3,0-0-0)[2]"), Y.GetPreparedData<float>()[2], 3.0f);
+		UTEST_EQUAL(TEXT("Slice(XC1x2x3,1x2x3,0-0-0)[3]"), Y.GetPreparedData<float>()[3], 4.0f);
+		UTEST_EQUAL(TEXT("Slice(XC1x2x3,1x2x3,0-0-0)[4]"), Y.GetPreparedData<float>()[4], 5.0f);
+		UTEST_EQUAL(TEXT("Slice(XC1x2x3,1x2x3,0-0-0)[5]"), Y.GetPreparedData<float>()[5], 6.0f);
+
+		Y = MakeTensor(TEXT("Y"), { 1,1,3 });
+		CallSliceApply(XC1x2x3, Y, { 0,1,0 });
+		UTEST_TRUE(TEXT("Y const if input is const"), Y.HasPreparedData());
+		UTEST_EQUAL(TEXT("Slice(XC1x2x3,1x1x3,0-1-0)[0]"), Y.GetPreparedData<float>()[0], 4.0f);
+		UTEST_EQUAL(TEXT("Slice(XC1x2x3,1x1x3,0-1-0)[1]"), Y.GetPreparedData<float>()[1], 5.0f);
+		UTEST_EQUAL(TEXT("Slice(XC1x2x3,1x1x3,0-1-0)[2]"), Y.GetPreparedData<float>()[2], 6.0f);
+
+		Y = MakeTensor(TEXT("Y"), { 1,2,2 });
+		CallSliceApply(XC1x2x3, Y, { 0,0,1 });
+		UTEST_TRUE(TEXT("Y const if input is const"), Y.HasPreparedData());
+		UTEST_EQUAL(TEXT("Slice(XC1x2x3,1x2x2,0-0-1)[0]"), Y.GetPreparedData<float>()[0], 2.0f);
+		UTEST_EQUAL(TEXT("Slice(XC1x2x3,1x2x2,0-0-1)[1]"), Y.GetPreparedData<float>()[1], 3.0f);
+		UTEST_EQUAL(TEXT("Slice(XC1x2x3,1x2x2,0-0-1)[2]"), Y.GetPreparedData<float>()[2], 5.0f);
+		UTEST_EQUAL(TEXT("Slice(XC1x2x3,1x2x2,0-0-1)[3]"), Y.GetPreparedData<float>()[3], 6.0f);
+
+		return true;
+	}
+
+	IMPLEMENT_NNE_SHAPEINFERENCEHELPER_UNIT_AUTOMATION_TEST(FSliceCPUHelperSteps, "System.Engine.MachineLearning.NNE.RDG.UnitTest.SliceHelper.Steps");
+	bool FSliceCPUHelperSteps::RunTest(const FString& Parameter)
+	{
+		FTensor XC1x2x3 = MakeConstTensor(TEXT("XC1x2x3"), { 1,2,3 }, { 1.0f, 2.0f, 3.0f, 4.0f, 5.0f, 6.0f });
+
+		FTensor Y = MakeTensor(TEXT("Y"), { 1,2,2 });
+		CallSliceApply(XC1x2x3, Y, { 0, 1, 0 }, { 1, -1, 2 });
+		UTEST_TRUE(TEXT("Y const if input is const"), Y.HasPreparedData());
+		UTEST_EQUAL(TEXT("Slice(XC1x2x3,1x2x3,0--1-0,1--1-2)[0]"), Y.GetPreparedData<float>()[0], 4.0f);
+		UTEST_EQUAL(TEXT("Slice(XC1x2x3,1x2x3,0--1-0,1--1-2)[1]"), Y.GetPreparedData<float>()[1], 6.0f);
+		UTEST_EQUAL(TEXT("Slice(XC1x2x3,1x2x3,0--1-0,1--1-2)[2]"), Y.GetPreparedData<float>()[2], 1.0f);
+		UTEST_EQUAL(TEXT("Slice(XC1x2x3,1x2x3,0--1-0,1--1-2)[3]"), Y.GetPreparedData<float>()[3], 3.0f);
+
+		return true;
+	}
+
+#endif //WITH_DEV_AUTOMATION_TESTS
+
+} // namespace UE::NNERuntimeRDG::Private::Test::SliceOp
diff --git a/Engine/Plugins/Experimental/NNERuntimeRDG/Source/NNERuntimeRDG/Private/Tests/NNERuntimeRDGUnitTestOpTranspose.cpp b/Engine/Plugins/Experimental/NNERuntimeRDG/Source/NNERuntimeRDG/Private/Tests/NNERuntimeRDGUnitTestOpTranspose.cpp
new file mode 100644
--- /dev/null
+++ b/Engine/Plugins/Experimental/NNERuntimeRDG/Source/NNERuntimeRDG/Private/Tests/NNERuntimeRDGUnitTestOpTranspose.cpp
@@ -0,0 +1,102 @@
+// Copyright Epic Games, Inc. All Rights Reserved.
+
+#include "CoreMinimal.h"
+#include "NNETensor.h"
+#include "NNETypes.h"
+#include "NNERuntimeRDGHelperTranspose.h"
+#include "NNERuntimeRDGUnitTestHelper.h"
+
+#if WITH_DEV_AUTOMATION_TESTS
+#include "Misc/AutomationTest.h"
+#endif
+
+namespace UE::NNERuntimeRDG::Private::Test::TransposeOp
+{
+
+#if WITH_DEV_AUTOMATION_TESTS
+
+	using namespace NNE::Internal;
+	using namespace UE::NNERuntimeRDG::Internal;
+
+	IMPLEMENT_NNE_SHAPEINFERENCEHELPER_UNIT_AUTOMATION_TEST(FTransposeCPUHelperIdentity, "System.Engine.MachineLearning.NNE.RDG.UnitTest.TransposeHelper.Identity");
+	bool FTransposeCPUHelperIdentity::RunTest(const FString& Parameter)
+	{
+		FTensor XC1x2 = MakeConstTensor(TEXT("XC1x2"), { 1,2 }, { 1.0f, 2.0f });
+		FTensor XC1x2x3 = MakeConstTensor(TEXT("XC1x2x3"), { 1,2,3 }, { 1.0f, 2.0f, 3.0f, 4.0f, 5.0f, 6.0f });
+
+		CPUHelper::Transpose::TransposePreparedData(XC1x2, { 0,1 });
+		UTEST_EQUAL(TEXT("Transpose(XC1x2,0,1)[0]"), XC1x2.GetPreparedData<float>()[0], 1.0f);
+		UTEST_EQUAL(TEXT("Transpose(XC1x2,0,1)[1]"), XC1x2.GetPreparedData<float>()[1], 2.0f);
+
+		CPUHelper::Transpose::TransposePreparedData(XC1x2x3, { 0,1,2 });
+		UTEST_EQUAL(TEXT("Transpose(XC1x2x3,0,1,2)[0]"), XC1x2x3.GetPreparedData<float>()[0], 1.0f);
+		UTEST_EQUAL(TEXT("Transpose(XC1x2x3,0,1,2)[1]"), XC1x2x3.GetPreparedData<float>()[1], 2.0f);
+		UTEST_EQUAL(TEXT("Transpose(XC1x2x3,0,1,2)[2]"), XC1x2x3.GetPreparedData<float>()[2], 3.0f);
+		UTEST_EQUAL(TEXT("Transpose(XC1x2x3,0,1,2)[3]"), XC1x2x3.GetPreparedData<float>()[3], 4.0f);
+		UTEST_EQUAL(TEXT("Transpose(XC1x2x3,0,1,2)[4]"), XC1x2x3.GetPreparedData<float>()[4], 5.0f);
+		UTEST_EQUAL(TEXT("Transpose(XC1x2x3,0,1,2)[5]"), XC1x2x3.GetPreparedData<float>()[5], 6.0f);
+
+		return true;
+	}
+
+	IMPLEMENT_NNE_SHAPEINFERENCEHELPER_UNIT_AUTOMATION_TEST(FTransposeCPUHelperReverse, "System.Engine.MachineLearning.NNE.RDG.UnitTest.TransposeHelper.Reverse");
+	bool FTransposeCPUHelperReverse::RunTest(const FString& Parameter)
+	{
+		FTensor XC1x2 = MakeConstTensor(TEXT("XC1x2"), { 1,2 }, { 1.0f, 2.0f });
+		FTensor XC1x2x3 = MakeConstTensor(TEXT("XC1x2x3"), { 1,2,3 }, { 1.0f, 2.0f, 3.0f, 4.0f, 5.0f, 6.0f });
+
+		CPUHelper::Transpose::TransposePreparedData(XC1x2, { 1,0 });
+		UTEST_EQUAL(TEXT("Transpose(XC1x2,1,0)[0]"), XC1x2.GetPreparedData<float>()[0], 1.0f);
+		UTEST_EQUAL(TEXT("Transpose(XC1x2,1,0)[1]"), XC1x2.GetPreparedData<float>()[1], 2.0f);
+
+		CPUHelper::Transpose::TransposePreparedData(XC1x2x3, { 2,1,0 });
+		UTEST_EQUAL(TEXT("Transpose(XC1x2x3,2,1,0)[0]"), XC1x2x3.GetPreparedData<float>()[0], 1.0f);
+		UTEST_EQUAL(TEXT("Transpose(XC1x2x3,2,1,0)[1]"), XC1x2x3.GetPreparedData<float>()[1], 4.0f);
+		UTEST_EQUAL(TEXT("Transpose(XC1x2x3,2,1,0)[2]"), XC1x2x3.GetPreparedData<float>()[2], 2.0f);
+		UTEST_EQUAL(TEXT("Transpose(XC1x2x3,2,1,0)[3]"), XC1x2x3.GetPreparedData<float>()[3], 5.0f);
+		UTEST_EQUAL(TEXT("Transpose(XC1x2x3,2,1,0)[4]"), XC1x2x3.GetPreparedData<float>()[4], 3.0f);
+		UTEST_EQUAL(TEXT("Transpose(XC1x2x3,2,1,0)[5]"), XC1x2x3.GetPreparedData<float>()[5], 6.0f);
+
+		return true;
+	}
+
+	IMPLEMENT_NNE_SHAPEINFERENCEHELPER_UNIT_AUTOMATION_TEST(FTransposeCPUHelperConvKernel, "System.Engine.MachineLearning.NNE.RDG.UnitTest.TransposeHelper.ConvKernel");
+	bool FTransposeCPUHelperConvKernel::RunTest(const FString& Parameter)
+	{
+		FTensor XC2x3x4x5 = MakeConstTensor(TEXT("XC2x3x4x5"), { 2,3,4,5 }, { 
+			0.0f,   1.0f,   2.0f,   3.0f,   4.0f,   5.0f,   6.0f,   7.0f,   8.0f,   9.0f,   10.0f,  11.0f,  12.0f,
+			13.0f,  14.0f,  15.0f,  16.0f,  17.0f,  18.0f,  19.0f,  20.0f,  21.0f,  22.0f,  23.0f,  24.0f,  25.0f,
+			26.0f,  27.0f,  28.0f,  29.0f,  30.0f,  31.0f,  32.0f,  33.0f,  34.0f,  35.0f,  36.0f,  37.0f,  38.0f,
+			39.0f,  40.0f,  41.0f,  42.0f,  43.0f,  44.0f,  45.0f,  46.0f,  47.0f,  48.0f,  49.0f,  50.0f,  51.0f,
+			52.0f,  53.0f,  54.0f,  55.0f,  56.0f,  57.0f,  58.0f,  59.0f,  60.0f,  61.0f,  62.0f,  63.0f,  64.0f,
+			65.0f,  66.0f,  67.0f,  68.0f,  69.0f,  70.0f,  71.0f,  72.0f,  73.0f,  74.0f,  75.0f,  76.0f,  77.0f,
+			78.0f,  79.0f,  80.0f,  81.0f,  82.0f,  83.0f,  84.0f,  85.0f,  86.0f,  87.0f,  88.0f,  89.0f,  90.0f,
+			91.0f,  92.0f,  93.0f,  94.0f,  95.0f,  96.0f,  97.0f,  98.0f,  99.0f, 100.0f, 101.0f, 102.0f, 103.0f,
+			104.0f,105.0f, 106.0f, 107.0f, 108.0f, 109.0f, 110.0f, 111.0f, 112.0f, 113.0f, 114.0f, 115.0f, 116.0f,
+			117.0f,118.0f, 119.0f});
+
+		TArray<float> ExpectedResult({
+			0.0f,   60.0f,  20.0f,  80.0f,  40.0f, 100.0f,   1.0f,  61.0f,  21.0f,  81.0f,  41.0f, 101.0f,   2.0f,
+			62.0f,  22.0f,  82.0f,  42.0f, 102.0f,   3.0f,  63.0f,  23.0f,  83.0f,  43.0f, 103.0f,   4.0f,  64.0f,
+			24.0f,  84.0f,  44.0f, 104.0f,   5.0f,  65.0f,  25.0f,  85.0f,  45.0f, 105.0f,   6.0f,  66.0f,  26.0f,
+			86.0f,  46.0f, 106.0f,   7.0f,  67.0f,  27.0f,  87.0f,  47.0f, 107.0f,   8.0f,  68.0f,  28.0f,  88.0f,
+			48.0f, 108.0f,   9.0f,  69.0f,  29.0f,  89.0f,  49.0f, 109.0f,  10.0f,  70.0f,  30.0f,  90.0f,  50.0f,
+			110.0f, 11.0f,  71.0f,  31.0f,  91.0f,  51.0f, 111.0f,  12.0f,  72.0f,  32.0f,  92.0f,  52.0f, 112.0f,
+			13.0f,  73.0f,  33.0f,  93.0f,  53.0f, 113.0f,  14.0f,  74.0f,  34.0f,  94.0f,  54.0f, 114.0f,  15.0f,
+			75.0f,  35.0f,  95.0f,  55.0f, 115.0f,  16.0f,  76.0f,  36.0f,  96.0f,  56.0f, 116.0f,  17.0f,  77.0f,
+			37.0f,  97.0f,  57.0f, 117.0f,  18.0f,  78.0f,  38.0f,  98.0f,  58.0f, 118.0f,  19.0f,  79.0f,  39.0f,
+			99.0f,  59.0f, 119 });
+
+		CPUHelper::Transpose::TransposePreparedData(XC2x3x4x5, { 2,3,1,0 });
+		for (uint32 i = 0; i < XC2x3x4x5.GetVolume(); ++i)
+		{
+			FString TestText = FString::Printf(TEXT("Transpose(XC2x3x4x5,2,3,1,0)[ %d]"), i);
+			UTEST_EQUAL(*TestText, XC2x3x4x5.GetPreparedData<float>()[i], ExpectedResult[i]);
+		}
+
+		return true;
+	}
+
+#endif //WITH_DEV_AUTOMATION_TESTS
+
+} // namespace UE::NNERuntimeRDG::Private::Test::TransposeOp
diff --git a/Engine/Plugins/Experimental/NNERuntimeRDG/Source/NNERuntimeRDG/Private/Tests/NNERuntimeRDGUnitTestOpUnary.cpp b/Engine/Plugins/Experimental/NNERuntimeRDG/Source/NNERuntimeRDG/Private/Tests/NNERuntimeRDGUnitTestOpUnary.cpp
new file mode 100644
--- /dev/null
+++ b/Engine/Plugins/Experimental/NNERuntimeRDG/Source/NNERuntimeRDG/Private/Tests/NNERuntimeRDGUnitTestOpUnary.cpp
@@ -0,0 +1,438 @@
+// Copyright Epic Games, Inc. All Rights Reserved.
+
+#include "CoreMinimal.h"
+#include "Misc/EnumRange.h"
+#include "NNETensor.h"
+#include "NNETypes.h"
+#include "NNERuntimeRDGHelperElementWiseUnary.h"
+#include "NNERuntimeRDGUnitTestHelper.h"
+
+#if WITH_DEV_AUTOMATION_TESTS
+#include "Misc/AutomationTest.h"
+#endif
+
+ENUM_RANGE_BY_COUNT(UE::NNEHlslShaders::Internal::EElementWiseUnaryOperatorType, UE::NNEHlslShaders::Internal::EElementWiseUnaryOperatorType::MAX)
+
+namespace UE::NNERuntimeRDG::Private::Test::UnaryOp
+{
+
+#if WITH_DEV_AUTOMATION_TESTS
+
+	using namespace NNE;
+	using namespace NNE::Internal;
+	using namespace UE::NNERuntimeRDG::Internal;
+	using EElementWiseUnaryOperatorType = NNEHlslShaders::Internal::EElementWiseUnaryOperatorType;
+
+	IMPLEMENT_NNE_SHAPEINFERENCEHELPER_UNIT_AUTOMATION_TEST(FElementWiseUnaryCPUHelperConstOuput, "System.Engine.MachineLearning.NNE.RDG.UnitTest.UnaryHelper.ConstOutput");
+	bool FElementWiseUnaryCPUHelperConstOuput::RunTest(const FString& Parameter)
+	{
+		for (EElementWiseUnaryOperatorType Op : TEnumRange<EElementWiseUnaryOperatorType>())
+		{
+			if (!TestUnaryOutputIsOnlyComputedWhenItShould(Op))
+			{
+				return false;
+			}
+		}
+		return true;
+	}
+
+	IMPLEMENT_NNE_SHAPEINFERENCEHELPER_UNIT_AUTOMATION_TEST(FElementWiseUnaryCPUHelperAbs, "System.Engine.MachineLearning.NNE.RDG.UnitTest.UnaryHelper.Abs");
+	bool FElementWiseUnaryCPUHelperAbs::RunTest(const FString& Parameter)
+	{
+		FTensor XC2x1 = MakeConstTensor(TEXT("XC2x1"), { 2,1 }, { -1.0f, 2.0f });
+		FTensor Y = MakeTensor(TEXT("Y"), { 2,1 });
+		CPUHelper::ElementWiseUnary::Apply(EElementWiseUnaryOperatorType::Abs, XC2x1, 1.0f, 2.0f, 3.0f, Y);
+		UTEST_TRUE(TEXT("Y const if input is const"), Y.HasPreparedData());
+		UTEST_EQUAL(TEXT("Abs(XC2x1)[0]"), Y.GetPreparedData<float>()[0], FMath::Abs(XC2x1.GetPreparedData<float>()[0]));
+		UTEST_EQUAL(TEXT("Abs(XC2x1)[1]"), Y.GetPreparedData<float>()[1], FMath::Abs(XC2x1.GetPreparedData<float>()[1]));
+		return true;
+	}
+
+	IMPLEMENT_NNE_SHAPEINFERENCEHELPER_UNIT_AUTOMATION_TEST(FElementWiseUnaryCPUHelperAcos, "System.Engine.MachineLearning.NNE.RDG.UnitTest.UnaryHelper.Acos");
+	bool FElementWiseUnaryCPUHelperAcos::RunTest(const FString& Parameter)
+	{
+		FTensor XC2x1 = MakeConstTensor(TEXT("XC2x1"), { 2,1 }, { -0.5f, 0.8f });
+		FTensor Y = MakeTensor(TEXT("Y"), { 2,1 });
+		CPUHelper::ElementWiseUnary::Apply(EElementWiseUnaryOperatorType::Acos, XC2x1, 1.0f, 2.0f, 3.0f, Y);
+		UTEST_TRUE(TEXT("Y const if input is const"), Y.HasPreparedData());
+		UTEST_EQUAL(TEXT("Acos(XC2x1)[0]"), Y.GetPreparedData<float>()[0], FMath::Acos(XC2x1.GetPreparedData<float>()[0]));
+		UTEST_EQUAL(TEXT("Acos(XC2x1)[1]"), Y.GetPreparedData<float>()[1], FMath::Acos(XC2x1.GetPreparedData<float>()[1]));
+		return true;
+	}
+
+	IMPLEMENT_NNE_SHAPEINFERENCEHELPER_UNIT_AUTOMATION_TEST(FElementWiseUnaryCPUHelperAcosh, "System.Engine.MachineLearning.NNE.RDG.UnitTest.UnaryHelper.Acosh");
+	bool FElementWiseUnaryCPUHelperAcosh::RunTest(const FString& Parameter)
+	{
+		FTensor XC2x1 = MakeConstTensor(TEXT("XC2x1"), { 2,1 }, { 1.0f, 1.5f });
+		FTensor Y = MakeTensor(TEXT("Y"), { 2,1 });
+		CPUHelper::ElementWiseUnary::Apply(EElementWiseUnaryOperatorType::Acosh, XC2x1, 1.0f, 2.0f, 3.0f, Y);
+		UTEST_TRUE(TEXT("Y const if input is const"), Y.HasPreparedData());
+		UTEST_EQUAL(TEXT("Acosh(XC2x1)[0]"), Y.GetPreparedData<float>()[0], 0.0f);
+		UTEST_EQUAL(TEXT("Acosh(XC2x1)[1]"), Y.GetPreparedData<float>()[1], 0.9624236501192069f);
+		return true;
+	}
+
+	IMPLEMENT_NNE_SHAPEINFERENCEHELPER_UNIT_AUTOMATION_TEST(FElementWiseUnaryCPUHelperAsin, "System.Engine.MachineLearning.NNE.RDG.UnitTest.UnaryHelper.Asin");
+	bool FElementWiseUnaryCPUHelperAsin::RunTest(const FString& Parameter)
+	{
+		FTensor XC2x1 = MakeConstTensor(TEXT("XC2x1"), { 2,1 }, { -0.5f, 0.8f });
+		FTensor Y = MakeTensor(TEXT("Y"), { 2,1 });
+		CPUHelper::ElementWiseUnary::Apply(EElementWiseUnaryOperatorType::Asin, XC2x1, 1.0f, 2.0f, 3.0f, Y);
+		UTEST_TRUE(TEXT("Y const if input is const"), Y.HasPreparedData());
+		UTEST_EQUAL(TEXT("Asin(XC2x1)[0]"), Y.GetPreparedData<float>()[0], FMath::Asin(XC2x1.GetPreparedData<float>()[0]));
+		UTEST_EQUAL(TEXT("Asin(XC2x1)[1]"), Y.GetPreparedData<float>()[1], FMath::Asin(XC2x1.GetPreparedData<float>()[1]));
+		return true;
+	}
+
+	IMPLEMENT_NNE_SHAPEINFERENCEHELPER_UNIT_AUTOMATION_TEST(FElementWiseUnaryCPUHelperAsinh, "System.Engine.MachineLearning.NNE.RDG.UnitTest.UnaryHelper.Asinh");
+	bool FElementWiseUnaryCPUHelperAsinh::RunTest(const FString& Parameter)
+	{
+		FTensor XC2x1 = MakeConstTensor(TEXT("XC2x1"), { 2,1 }, { 1.0f, 2.0f });
+		FTensor Y = MakeTensor(TEXT("Y"), { 2,1 });
+		CPUHelper::ElementWiseUnary::Apply(EElementWiseUnaryOperatorType::Asinh, XC2x1, 1.0f, 2.0f, 3.0f, Y);
+		UTEST_TRUE(TEXT("Y const if input is const"), Y.HasPreparedData());
+		UTEST_EQUAL(TEXT("Asinh(XC2x1)[0]"), Y.GetPreparedData<float>()[0], 0.881373587019543f);
+		UTEST_EQUAL(TEXT("Asinh(XC2x1)[1]"), Y.GetPreparedData<float>()[1], 1.4436354751788103f);
+		return true;
+	}
+
+	IMPLEMENT_NNE_SHAPEINFERENCEHELPER_UNIT_AUTOMATION_TEST(FElementWiseUnaryCPUHelperAtan, "System.Engine.MachineLearning.NNE.RDG.UnitTest.UnaryHelper.Atan");
+	bool FElementWiseUnaryCPUHelperAtan::RunTest(const FString& Parameter)
+	{
+		FTensor XC2x1 = MakeConstTensor(TEXT("XC2x1"), { 2,1 }, { -0.5f, 0.8f });
+		FTensor Y = MakeTensor(TEXT("Y"), { 2,1 });
+		CPUHelper::ElementWiseUnary::Apply(EElementWiseUnaryOperatorType::Atan, XC2x1, 1.0f, 2.0f, 3.0f, Y);
+		UTEST_TRUE(TEXT("Y const if input is const"), Y.HasPreparedData());
+		UTEST_EQUAL(TEXT("Atan(XC2x1)[0]"), Y.GetPreparedData<float>()[0], FMath::Atan(XC2x1.GetPreparedData<float>()[0]));
+		UTEST_EQUAL(TEXT("Atan(XC2x1)[1]"), Y.GetPreparedData<float>()[1], FMath::Atan(XC2x1.GetPreparedData<float>()[1]));
+		return true;
+	}
+
+	IMPLEMENT_NNE_SHAPEINFERENCEHELPER_UNIT_AUTOMATION_TEST(FElementWiseUnaryCPUHelperCeil, "System.Engine.MachineLearning.NNE.RDG.UnitTest.UnaryHelper.Ceil");
+	bool FElementWiseUnaryCPUHelperCeil::RunTest(const FString& Parameter)
+	{
+		FTensor XC2x1 = MakeConstTensor(TEXT("XC2x1"), { 2,1 }, { -0.5f, 0.8f });
+		FTensor Y = MakeTensor(TEXT("Y"), { 2,1 });
+		CPUHelper::ElementWiseUnary::Apply(EElementWiseUnaryOperatorType::Ceil, XC2x1, 1.0f, 2.0f, 3.0f, Y);
+		UTEST_TRUE(TEXT("Y const if input is const"), Y.HasPreparedData());
+		UTEST_EQUAL(TEXT("Ceil(XC2x1)[0]"), Y.GetPreparedData<float>()[0], FMath::CeilToFloat(XC2x1.GetPreparedData<float>()[0]));
+		UTEST_EQUAL(TEXT("Ceil(XC2x1)[1]"), Y.GetPreparedData<float>()[1], FMath::CeilToFloat(XC2x1.GetPreparedData<float>()[1]));
+		return true;
+	}
+
+	IMPLEMENT_NNE_SHAPEINFERENCEHELPER_UNIT_AUTOMATION_TEST(FElementWiseUnaryCPUHelperClip, "System.Engine.MachineLearning.NNE.RDG.UnitTest.UnaryHelper.Clip");
+	bool FElementWiseUnaryCPUHelperClip::RunTest(const FString& Parameter)
+	{
+		FTensor XC2x1 = MakeConstTensor(TEXT("XC3x1"), { 3,1 }, { 0.5f, 1.8f, 2.1f });
+		FTensor Y = MakeTensor(TEXT("Y"), { 3,1 });
+		CPUHelper::ElementWiseUnary::Apply(EElementWiseUnaryOperatorType::Clip, XC2x1, 1.0f, 2.0f, 3.0f, Y);
+		UTEST_TRUE(TEXT("Y const if input is const"), Y.HasPreparedData());
+		UTEST_EQUAL(TEXT("Clip(XC3x1)[0]"), Y.GetPreparedData<float>()[0], 1.0f);
+		UTEST_EQUAL(TEXT("Clip(XC3x1)[1]"), Y.GetPreparedData<float>()[1], XC2x1.GetPreparedData<float>()[1]);
+		UTEST_EQUAL(TEXT("Clip(XC3x1)[2]"), Y.GetPreparedData<float>()[2], 2.0f);
+		return true;
+	}
+
+	IMPLEMENT_NNE_SHAPEINFERENCEHELPER_UNIT_AUTOMATION_TEST(FElementWiseUnaryCPUHelperCos, "System.Engine.MachineLearning.NNE.RDG.UnitTest.UnaryHelper.Cos");
+	bool FElementWiseUnaryCPUHelperCos::RunTest(const FString& Parameter)
+	{
+		FTensor XC2x1 = MakeConstTensor(TEXT("XC2x1"), { 2,1 }, { -0.5f, 0.8f });
+		FTensor Y = MakeTensor(TEXT("Y"), { 2,1 });
+		CPUHelper::ElementWiseUnary::Apply(EElementWiseUnaryOperatorType::Cos, XC2x1, 1.0f, 2.0f, 3.0f, Y);
+		UTEST_TRUE(TEXT("Y const if input is const"), Y.HasPreparedData());
+		UTEST_EQUAL(TEXT("Cos(XC2x1)[0]"), Y.GetPreparedData<float>()[0], FMath::Cos(XC2x1.GetPreparedData<float>()[0]));
+		UTEST_EQUAL(TEXT("Cos(XC2x1)[1]"), Y.GetPreparedData<float>()[1], FMath::Cos(XC2x1.GetPreparedData<float>()[1]));
+		return true;
+	}
+
+	IMPLEMENT_NNE_SHAPEINFERENCEHELPER_UNIT_AUTOMATION_TEST(FElementWiseUnaryCPUHelperCosh, "System.Engine.MachineLearning.NNE.RDG.UnitTest.UnaryHelper.Cosh");
+	bool FElementWiseUnaryCPUHelperCosh::RunTest(const FString& Parameter)
+	{
+		FTensor XC2x1 = MakeConstTensor(TEXT("XC2x1"), { 2,1 }, { 1.0f, 2.0f });
+		FTensor Y = MakeTensor(TEXT("Y"), { 2,1 });
+		CPUHelper::ElementWiseUnary::Apply(EElementWiseUnaryOperatorType::Cosh, XC2x1, 1.0f, 2.0f, 3.0f, Y);
+		UTEST_TRUE(TEXT("Y const if input is const"), Y.HasPreparedData());
+		UTEST_EQUAL(TEXT("Cosh(XC2x1)[0]"), Y.GetPreparedData<float>()[0], 1.5430806348152437f);
+		UTEST_EQUAL(TEXT("Cosh(XC2x1)[1]"), Y.GetPreparedData<float>()[1], 3.7621956910836314f);
+		return true;
+	}
+	IMPLEMENT_NNE_SHAPEINFERENCEHELPER_UNIT_AUTOMATION_TEST(FElementWiseUnaryCPUHelperElu, "System.Engine.MachineLearning.NNE.RDG.UnitTest.UnaryHelper.Elu");
+	bool FElementWiseUnaryCPUHelperElu::RunTest(const FString& Parameter)
+	{
+		FTensor XC2x1 = MakeConstTensor(TEXT("XC2x1"), { 2,1 }, { -0.5f, 0.8f });
+		FTensor Y = MakeTensor(TEXT("Y"), { 2,1 });
+		CPUHelper::ElementWiseUnary::Apply(EElementWiseUnaryOperatorType::Elu, XC2x1, 1.0f, 2.0f, 3.0f, Y);
+		UTEST_TRUE(TEXT("Y const if input is const"), Y.HasPreparedData());
+		UTEST_EQUAL(TEXT("Elu(XC2x1)[0]"), Y.GetPreparedData<float>()[0], -0.3934693402873666f);
+		UTEST_EQUAL(TEXT("Elu(XC2x1)[1]"), Y.GetPreparedData<float>()[1], 0.8f);
+		return true;
+	}
+	
+	IMPLEMENT_NNE_SHAPEINFERENCEHELPER_UNIT_AUTOMATION_TEST(FElementWiseUnaryCPUHelperErf, "System.Engine.MachineLearning.NNE.RDG.UnitTest.UnaryHelper.Erf");
+	bool FElementWiseUnaryCPUHelperErf::RunTest(const FString& Parameter)
+	{
+		FTensor XC2x1 = MakeConstTensor(TEXT("XC2x1"), { 2,1 }, { 1.0f, 2.0f });
+		FTensor Y = MakeTensor(TEXT("Y"), { 2,1 });
+		CPUHelper::ElementWiseUnary::Apply(EElementWiseUnaryOperatorType::Erf, XC2x1, 1.0f, 2.0f, 3.0f, Y);
+		UTEST_TRUE(TEXT("Y const if input is const"), Y.HasPreparedData());
+		UTEST_EQUAL_TOLERANCE(TEXT("Erf(XC2x1)[0]"), Y.GetPreparedData<float>()[0], 0.8427007929497149f, 1e-3f);
+		UTEST_EQUAL_TOLERANCE(TEXT("Erf(XC2x1)[1]"), Y.GetPreparedData<float>()[1], 0.9953222650189527f, 1e-3f);
+		return true;
+	}
+
+	IMPLEMENT_NNE_SHAPEINFERENCEHELPER_UNIT_AUTOMATION_TEST(FElementWiseUnaryCPUHelperExp, "System.Engine.MachineLearning.NNE.RDG.UnitTest.UnaryHelper.Exp");
+	bool FElementWiseUnaryCPUHelperExp::RunTest(const FString& Parameter)
+	{
+		FTensor XC2x1 = MakeConstTensor(TEXT("XC2x1"), { 2,1 }, { -0.5f, 0.8f });
+		FTensor Y = MakeTensor(TEXT("Y"), { 2,1 });
+		CPUHelper::ElementWiseUnary::Apply(EElementWiseUnaryOperatorType::Exp, XC2x1, 1.0f, 2.0f, 3.0f, Y);
+		UTEST_TRUE(TEXT("Y const if input is const"), Y.HasPreparedData());
+		UTEST_EQUAL(TEXT("Exp(XC2x1)[0]"), Y.GetPreparedData<float>()[0], FMath::Exp(XC2x1.GetPreparedData<float>()[0]));
+		UTEST_EQUAL(TEXT("Exp(XC2x1)[1]"), Y.GetPreparedData<float>()[1], FMath::Exp(XC2x1.GetPreparedData<float>()[1]));
+		return true;
+	}
+
+	IMPLEMENT_NNE_SHAPEINFERENCEHELPER_UNIT_AUTOMATION_TEST(FElementWiseUnaryCPUHelperFloor, "System.Engine.MachineLearning.NNE.RDG.UnitTest.UnaryHelper.Floor");
+	bool FElementWiseUnaryCPUHelperFloor::RunTest(const FString& Parameter)
+	{
+		FTensor XC2x1 = MakeConstTensor(TEXT("XC2x1"), { 2,1 }, { -0.5f, 0.8f });
+		FTensor Y = MakeTensor(TEXT("Y"), { 2,1 });
+		CPUHelper::ElementWiseUnary::Apply(EElementWiseUnaryOperatorType::Floor, XC2x1, 1.0f, 2.0f, 3.0f, Y);
+		UTEST_TRUE(TEXT("Y const if input is const"), Y.HasPreparedData());
+		UTEST_EQUAL(TEXT("Floor(XC2x1)[0]"), Y.GetPreparedData<float>()[0], FMath::Floor(XC2x1.GetPreparedData<float>()[0]));
+		UTEST_EQUAL(TEXT("Floor(XC2x1)[1]"), Y.GetPreparedData<float>()[1], FMath::Floor(XC2x1.GetPreparedData<float>()[1]));
+		return true;
+	}
+	
+	IMPLEMENT_NNE_SHAPEINFERENCEHELPER_UNIT_AUTOMATION_TEST(FElementWiseUnaryCPUHelperIsInf, "System.Engine.MachineLearning.NNE.RDG.UnitTest.UnaryHelper.IsInf");
+	bool FElementWiseUnaryCPUHelperIsInf::RunTest(const FString& Parameter)
+	{
+		FTensor XC2x1 = MakeConstTensor(TEXT("XC2x1"), { 2,1 }, { -0.5f, INFINITY });
+		FTensor Y = MakeTensor(TEXT("Y"), { 2,1 });
+		CPUHelper::ElementWiseUnary::Apply(EElementWiseUnaryOperatorType::IsInf, XC2x1, 1.0f, 2.0f, 3.0f, Y);
+		UTEST_TRUE(TEXT("Y const if input is const"), Y.HasPreparedData());
+		UTEST_EQUAL(TEXT("IsInf(XC2x1)[0]"), Y.GetPreparedData<float>()[0], 0.0f);
+		UTEST_EQUAL(TEXT("IsInf(XC2x1)[1]"), Y.GetPreparedData<float>()[1], 1.0f);
+		return true;
+	}
+
+	IMPLEMENT_NNE_SHAPEINFERENCEHELPER_UNIT_AUTOMATION_TEST(FElementWiseUnaryCPUHelperIsNan, "System.Engine.MachineLearning.NNE.RDG.UnitTest.UnaryHelper.IsNan");
+	bool FElementWiseUnaryCPUHelperIsNan::RunTest(const FString& Parameter)
+	{
+		float FloatNan = FMath::Sqrt(-1.0f);
+		FTensor XC2x1 = MakeConstTensor(TEXT("XC2x1"), { 2,1 }, { -0.5f, FloatNan });
+		FTensor Y = MakeTensor(TEXT("Y"), { 2,1 });
+		CPUHelper::ElementWiseUnary::Apply(EElementWiseUnaryOperatorType::IsNan, XC2x1, 1.0f, 2.0f, 3.0f, Y);
+		UTEST_TRUE(TEXT("Y const if input is const"), Y.HasPreparedData());
+		UTEST_EQUAL(TEXT("IsNan(XC2x1)[0]"), Y.GetPreparedData<float>()[0], 0.0f);
+		UTEST_EQUAL(TEXT("IsNan(XC2x1)[1]"), Y.GetPreparedData<float>()[1], 1.0f);
+		return true;
+	}
+
+	IMPLEMENT_NNE_SHAPEINFERENCEHELPER_UNIT_AUTOMATION_TEST(FElementWiseUnaryCPUHelperHardSigmoid, "System.Engine.MachineLearning.NNE.RDG.UnitTest.UnaryHelper.HardSigmoid");
+	bool FElementWiseUnaryCPUHelperHardSigmoid::RunTest(const FString& Parameter)
+	{
+		FTensor XC2x1 = MakeConstTensor(TEXT("XC2x1"), { 2,1 }, { 0.0f, -1.5f });
+		FTensor Y = MakeTensor(TEXT("Y"), { 2,1 });
+		CPUHelper::ElementWiseUnary::Apply(EElementWiseUnaryOperatorType::HardSigmoid, XC2x1, 1.0f, 2.0f, 3.0f, Y);
+		UTEST_TRUE(TEXT("Y const if input is const"), Y.HasPreparedData());
+		UTEST_EQUAL(TEXT("HardSigmoid(XC2x1)[0]"), Y.GetPreparedData<float>()[0], 1.0f);
+		UTEST_EQUAL(TEXT("HardSigmoid(XC2x1)[1]"), Y.GetPreparedData<float>()[1], 0.5f);
+		return true;
+	}
+
+	IMPLEMENT_NNE_SHAPEINFERENCEHELPER_UNIT_AUTOMATION_TEST(FElementWiseUnaryCPUHelperHardswitch, "System.Engine.MachineLearning.NNE.RDG.UnitTest.UnaryHelper.HardSwitch");
+	bool FElementWiseUnaryCPUHelperHardswitch::RunTest(const FString& Parameter)
+	{
+		FTensor XC2x1 = MakeConstTensor(TEXT("XC2x1"), { 2,1 }, { 0.0f, 1.0f });
+		FTensor Y = MakeTensor(TEXT("Y"), { 2,1 });
+		CPUHelper::ElementWiseUnary::Apply(EElementWiseUnaryOperatorType::HardSwish, XC2x1, 1.0f, 2.0f, 3.0f, Y);
+		UTEST_TRUE(TEXT("Y const if input is const"), Y.HasPreparedData());
+		UTEST_EQUAL(TEXT("HardSwitch(XC2x1)[0]"), Y.GetPreparedData<float>()[0], 0.5f);
+		UTEST_EQUAL(TEXT("HardSwitch(XC2x1)[1]"), Y.GetPreparedData<float>()[1], (1.0f / 6.0f) + 0.5f);
+		return true;
+	}
+
+	IMPLEMENT_NNE_SHAPEINFERENCEHELPER_UNIT_AUTOMATION_TEST(FElementWiseUnaryCPUHelperLeakyRelu, "System.Engine.MachineLearning.NNE.RDG.UnitTest.UnaryHelper.LeakyRelu");
+	bool FElementWiseUnaryCPUHelperLeakyRelu::RunTest(const FString& Parameter)
+	{
+		FTensor XC2x1 = MakeConstTensor(TEXT("XC2x1"), { 2,1 }, { 1.0f, -1.0f });
+		FTensor Y = MakeTensor(TEXT("Y"), { 2,1 });
+		CPUHelper::ElementWiseUnary::Apply(EElementWiseUnaryOperatorType::LeakyRelu, XC2x1, -0.5f, 2.0f, 3.0f, Y);
+		UTEST_TRUE(TEXT("Y const if input is const"), Y.HasPreparedData());
+		UTEST_EQUAL(TEXT("LeakyRelu(XC2x1)[0]"), Y.GetPreparedData<float>()[0], 1.0f);
+		UTEST_EQUAL(TEXT("LeakyRelu(XC2x1)[1]"), Y.GetPreparedData<float>()[1], 0.5f);
+		return true;
+	}
+
+	IMPLEMENT_NNE_SHAPEINFERENCEHELPER_UNIT_AUTOMATION_TEST(FElementWiseUnaryCPUHelperLog, "System.Engine.MachineLearning.NNE.RDG.UnitTest.UnaryHelper.Log");
+	bool FElementWiseUnaryCPUHelperLog::RunTest(const FString& Parameter)
+	{
+		FTensor XC2x1 = MakeConstTensor(TEXT("XC2x1"), { 2,1 }, { 0.5f, 0.8f });
+		FTensor Y = MakeTensor(TEXT("Y"), { 2,1 });
+		CPUHelper::ElementWiseUnary::Apply(EElementWiseUnaryOperatorType::Log, XC2x1, 1.0f, 2.0f, 3.0f, Y);
+		UTEST_TRUE(TEXT("Y const if input is const"), Y.HasPreparedData());
+		UTEST_EQUAL(TEXT("Log(XC2x1)[0]"), Y.GetPreparedData<float>()[0], FMath::Loge(XC2x1.GetPreparedData<float>()[0]));
+		UTEST_EQUAL(TEXT("Log(XC2x1)[1]"), Y.GetPreparedData<float>()[1], FMath::Loge(XC2x1.GetPreparedData<float>()[1]));
+		return true;
+	}
+
+	IMPLEMENT_NNE_SHAPEINFERENCEHELPER_UNIT_AUTOMATION_TEST(FElementWiseUnaryCPUHelperNeg, "System.Engine.MachineLearning.NNE.RDG.UnitTest.UnaryHelper.Neg");
+	bool FElementWiseUnaryCPUHelperNeg::RunTest(const FString& Parameter)
+	{
+		FTensor XC2x1 = MakeConstTensor(TEXT("XC2x1"), { 2,1 }, { -0.5f, 0.8f });
+		FTensor Y = MakeTensor(TEXT("Y"), { 2,1 });
+		CPUHelper::ElementWiseUnary::Apply(EElementWiseUnaryOperatorType::Neg, XC2x1, 1.0f, 2.0f, 3.0f, Y);
+		UTEST_TRUE(TEXT("Y const if input is const"), Y.HasPreparedData());
+		UTEST_EQUAL(TEXT("Neg(XC2x1)[0]"), Y.GetPreparedData<float>()[0], -XC2x1.GetPreparedData<float>()[0]);
+		UTEST_EQUAL(TEXT("Neg(XC2x1)[1]"), Y.GetPreparedData<float>()[1], -XC2x1.GetPreparedData<float>()[1]);
+		return true;
+	}
+
+	IMPLEMENT_NNE_SHAPEINFERENCEHELPER_UNIT_AUTOMATION_TEST(FElementWiseUnaryCPUHelperReciprocal, "System.Engine.MachineLearning.NNE.RDG.UnitTest.UnaryHelper.Reciprocal");
+	bool FElementWiseUnaryCPUHelperReciprocal::RunTest(const FString& Parameter)
+	{
+		FTensor XC2x1 = MakeConstTensor(TEXT("XC2x1"), { 2,1 }, { -0.5f, 0.8f });
+		FTensor Y = MakeTensor(TEXT("Y"), { 2,1 });
+		CPUHelper::ElementWiseUnary::Apply(EElementWiseUnaryOperatorType::Reciprocal, XC2x1, 1.0f, 2.0f, 3.0f, Y);
+		UTEST_TRUE(TEXT("Y const if input is const"), Y.HasPreparedData());
+		UTEST_EQUAL(TEXT("Reciprocal(XC2x1)[0]"), Y.GetPreparedData<float>()[0], 1.0f / (XC2x1.GetPreparedData<float>()[0]));
+		UTEST_EQUAL(TEXT("Reciprocal(XC2x1)[1]"), Y.GetPreparedData<float>()[1], 1.0f / (XC2x1.GetPreparedData<float>()[1]));
+		return true;
+	}
+
+	IMPLEMENT_NNE_SHAPEINFERENCEHELPER_UNIT_AUTOMATION_TEST(FElementWiseUnaryCPUHelperRelu, "System.Engine.MachineLearning.NNE.RDG.UnitTest.UnaryHelper.Relu");
+	bool FElementWiseUnaryCPUHelperRelu::RunTest(const FString& Parameter)
+	{
+		FTensor XC2x1 = MakeConstTensor(TEXT("XC2x1"), { 2,1 }, { -0.5f, 0.8f });
+		FTensor Y = MakeTensor(TEXT("Y"), { 2,1 });
+		CPUHelper::ElementWiseUnary::Apply(EElementWiseUnaryOperatorType::Relu, XC2x1, 1.0f, 2.0f, 3.0f, Y);
+		UTEST_TRUE(TEXT("Y const if input is const"), Y.HasPreparedData());
+		UTEST_EQUAL(TEXT("Relu(XC2x1)[0]"), Y.GetPreparedData<float>()[0], 0.0f);
+		UTEST_EQUAL(TEXT("Relu(XC2x1)[1]"), Y.GetPreparedData<float>()[1], 0.8f);
+		return true;
+	}
+
+	IMPLEMENT_NNE_SHAPEINFERENCEHELPER_UNIT_AUTOMATION_TEST(FElementWiseUnaryCPUHelperSelu, "System.Engine.MachineLearning.NNE.RDG.UnitTest.UnaryHelper.Selu");
+	bool FElementWiseUnaryCPUHelperSelu::RunTest(const FString& Parameter)
+	{
+		FTensor XC2x1 = MakeConstTensor(TEXT("XC2x1"), { 2,1 }, { 0.1f, -1.0f });
+		FTensor Y = MakeTensor(TEXT("Y"), { 2,1 });
+		CPUHelper::ElementWiseUnary::Apply(EElementWiseUnaryOperatorType::Selu, XC2x1, 1.0f, 2.0f, 3.0f, Y);
+		UTEST_TRUE(TEXT("Y const if input is const"), Y.HasPreparedData());
+		UTEST_EQUAL(TEXT("Selu(XC2x1)[0]"), Y.GetPreparedData<float>()[0], 0.3f);
+		UTEST_EQUAL(TEXT("Selu(XC2x1)[1]"), Y.GetPreparedData<float>()[1], 3.0f * (0.36787944117144233f - 1.0f));
+		return true;
+	}
+
+	IMPLEMENT_NNE_SHAPEINFERENCEHELPER_UNIT_AUTOMATION_TEST(FElementWiseUnaryCPUHelperSigmoid, "System.Engine.MachineLearning.NNE.RDG.UnitTest.UnaryHelper.Sigmoid");
+	bool FElementWiseUnaryCPUHelperSigmoid::RunTest(const FString& Parameter)
+	{
+		FTensor XC2x1 = MakeConstTensor(TEXT("XC2x1"), { 2,1 }, { 0.0f, 1.0f });
+		FTensor Y = MakeTensor(TEXT("Y"), { 2,1 });
+		CPUHelper::ElementWiseUnary::Apply(EElementWiseUnaryOperatorType::Sigmoid, XC2x1, 1.0f, 2.0f, 3.0f, Y);
+		UTEST_TRUE(TEXT("Y const if input is const"), Y.HasPreparedData());
+		UTEST_EQUAL(TEXT("Sigmoid(XC2x1)[0]"), Y.GetPreparedData<float>()[0], 0.5f);
+		UTEST_EQUAL(TEXT("Sigmoid(XC2x1)[1]"), Y.GetPreparedData<float>()[1], 1.0f / (1.0f + 0.36787944117144233f));
+		return true;
+	}
+
+	IMPLEMENT_NNE_SHAPEINFERENCEHELPER_UNIT_AUTOMATION_TEST(FElementWiseUnaryCPUHelperSign, "System.Engine.MachineLearning.NNE.RDG.UnitTest.UnaryHelper.Sign");
+	bool FElementWiseUnaryCPUHelperSign::RunTest(const FString& Parameter)
+	{
+		FTensor XC2x1 = MakeConstTensor(TEXT("XC2x1"), { 2,1 }, { -0.5f, 0.8f });
+		FTensor Y = MakeTensor(TEXT("Y"), { 2,1 });
+		CPUHelper::ElementWiseUnary::Apply(EElementWiseUnaryOperatorType::Sign, XC2x1, 1.0f, 2.0f, 3.0f, Y);
+		UTEST_TRUE(TEXT("Y const if input is const"), Y.HasPreparedData());
+		UTEST_EQUAL(TEXT("Sign(XC2x1)[0]"), Y.GetPreparedData<float>()[0], FMath::Sign(XC2x1.GetPreparedData<float>()[0]));
+		UTEST_EQUAL(TEXT("Sign(XC2x1)[1]"), Y.GetPreparedData<float>()[1], FMath::Sign(XC2x1.GetPreparedData<float>()[1]));
+		return true;
+	}
+
+	IMPLEMENT_NNE_SHAPEINFERENCEHELPER_UNIT_AUTOMATION_TEST(FElementWiseUnaryCPUHelperSin, "System.Engine.MachineLearning.NNE.RDG.UnitTest.UnaryHelper.Sin");
+	bool FElementWiseUnaryCPUHelperSin::RunTest(const FString& Parameter)
+	{
+		FTensor XC2x1 = MakeConstTensor(TEXT("XC2x1"), { 2,1 }, { -0.5f, 0.8f });
+		FTensor Y = MakeTensor(TEXT("Y"), { 2,1 });
+		CPUHelper::ElementWiseUnary::Apply(EElementWiseUnaryOperatorType::Sin, XC2x1, 1.0f, 2.0f, 3.0f, Y);
+		UTEST_TRUE(TEXT("Y const if input is const"), Y.HasPreparedData());
+		UTEST_EQUAL(TEXT("Sin(XC2x1)[0]"), Y.GetPreparedData<float>()[0], FMath::Sin(XC2x1.GetPreparedData<float>()[0]));
+		UTEST_EQUAL(TEXT("Sin(XC2x1)[1]"), Y.GetPreparedData<float>()[1], FMath::Sin(XC2x1.GetPreparedData<float>()[1]));
+		return true;
+	}
+
+	IMPLEMENT_NNE_SHAPEINFERENCEHELPER_UNIT_AUTOMATION_TEST(FElementWiseUnaryCPUHelperSinh, "System.Engine.MachineLearning.NNE.RDG.UnitTest.UnaryHelper.Sinh");
+	bool FElementWiseUnaryCPUHelperSinh::RunTest(const FString& Parameter)
+	{
+		FTensor XC2x1 = MakeConstTensor(TEXT("XC2x1"), { 2,1 }, { -0.5f, 0.8f });
+		FTensor Y = MakeTensor(TEXT("Y"), { 2,1 });
+		CPUHelper::ElementWiseUnary::Apply(EElementWiseUnaryOperatorType::Sinh, XC2x1, 1.0f, 2.0f, 3.0f, Y);
+		UTEST_TRUE(TEXT("Y const if input is const"), Y.HasPreparedData());
+		UTEST_EQUAL(TEXT("Sinh(XC2x1)[0]"), Y.GetPreparedData<float>()[0], FMath::Sinh(XC2x1.GetPreparedData<float>()[0]));
+		UTEST_EQUAL(TEXT("Sinh(XC2x1)[1]"), Y.GetPreparedData<float>()[1], FMath::Sinh(XC2x1.GetPreparedData<float>()[1]));
+		return true;
+	}
+
+	IMPLEMENT_NNE_SHAPEINFERENCEHELPER_UNIT_AUTOMATION_TEST(FElementWiseUnaryCPUHelperSoftplus, "System.Engine.MachineLearning.NNE.RDG.UnitTest.UnaryHelper.Softplus");
+	bool FElementWiseUnaryCPUHelperSoftplus::RunTest(const FString& Parameter)
+	{
+		FTensor XC2x1 = MakeConstTensor(TEXT("XC2x1"), { 2,1 }, { 0.0f, 1.0f });
+		FTensor Y = MakeTensor(TEXT("Y"), { 2,1 });
+		CPUHelper::ElementWiseUnary::Apply(EElementWiseUnaryOperatorType::Softplus, XC2x1, 1.0f, 2.0f, 3.0f, Y);
+		UTEST_TRUE(TEXT("Y const if input is const"), Y.HasPreparedData());
+		UTEST_EQUAL(TEXT("Softplus(XC2x1)[0]"), Y.GetPreparedData<float>()[0], 0.6931471805599453f);
+		UTEST_EQUAL(TEXT("Softplus(XC2x1)[1]"), Y.GetPreparedData<float>()[1], 1.3132616875182228f);
+		return true;
+	}
+
+	IMPLEMENT_NNE_SHAPEINFERENCEHELPER_UNIT_AUTOMATION_TEST(FElementWiseUnaryCPUHelperSoftsign, "System.Engine.MachineLearning.NNE.RDG.UnitTest.UnaryHelper.Softsign");
+	bool FElementWiseUnaryCPUHelperSoftsign::RunTest(const FString& Parameter)
+	{
+		FTensor XC2x1 = MakeConstTensor(TEXT("XC2x1"), { 2,1 }, { 0.0f, 1.0f });
+		FTensor Y = MakeTensor(TEXT("Y"), { 2,1 });
+		CPUHelper::ElementWiseUnary::Apply(EElementWiseUnaryOperatorType::Softsign, XC2x1, 1.0f, 2.0f, 3.0f, Y);
+		UTEST_TRUE(TEXT("Y const if input is const"), Y.HasPreparedData());
+		UTEST_EQUAL(TEXT("Softsign(XC2x1)[0]"), Y.GetPreparedData<float>()[0], 0.0f);
+		UTEST_EQUAL(TEXT("Softsign(XC2x1)[1]"), Y.GetPreparedData<float>()[1], 0.5f);
+		return true;
+	}
+
+	IMPLEMENT_NNE_SHAPEINFERENCEHELPER_UNIT_AUTOMATION_TEST(FElementWiseUnaryCPUHelperSqrt, "System.Engine.MachineLearning.NNE.RDG.UnitTest.UnaryHelper.Sqrt");
+	bool FElementWiseUnaryCPUHelperSqrt::RunTest(const FString& Parameter)
+	{
+		FTensor XC2x1 = MakeConstTensor(TEXT("XC2x1"), { 2,1 }, { 2.0f, 1.0f });
+		FTensor Y = MakeTensor(TEXT("Y"), { 2,1 });
+		CPUHelper::ElementWiseUnary::Apply(EElementWiseUnaryOperatorType::Sqrt, XC2x1, 1.0f, 2.0f, 3.0f, Y);
+		UTEST_TRUE(TEXT("Y const if input is const"), Y.HasPreparedData());
+		UTEST_EQUAL(TEXT("Sqrt(XC2x1)[0]"), Y.GetPreparedData<float>()[0], FMath::Sqrt(XC2x1.GetPreparedData<float>()[0]));
+		UTEST_EQUAL(TEXT("Sqrt(XC2x1)[1]"), Y.GetPreparedData<float>()[1], FMath::Sqrt(XC2x1.GetPreparedData<float>()[1]));
+		return true;
+	}
+
+	IMPLEMENT_NNE_SHAPEINFERENCEHELPER_UNIT_AUTOMATION_TEST(FElementWiseUnaryCPUHelperTan, "System.Engine.MachineLearning.NNE.RDG.UnitTest.UnaryHelper.Tan");
+	bool FElementWiseUnaryCPUHelperTan::RunTest(const FString& Parameter)
+	{
+		FTensor XC2x1 = MakeConstTensor(TEXT("XC2x1"), { 2,1 }, { -0.5f, 0.8f });
+		FTensor Y = MakeTensor(TEXT("Y"), { 2,1 });
+		CPUHelper::ElementWiseUnary::Apply(EElementWiseUnaryOperatorType::Tan, XC2x1, 1.0f, 2.0f, 3.0f, Y);
+		UTEST_TRUE(TEXT("Y const if input is const"), Y.HasPreparedData());
+		UTEST_EQUAL(TEXT("Tan(XC2x1)[0]"), Y.GetPreparedData<float>()[0], FMath::Tan(XC2x1.GetPreparedData<float>()[0]));
+		UTEST_EQUAL(TEXT("Tan(XC2x1)[1]"), Y.GetPreparedData<float>()[1], FMath::Tan(XC2x1.GetPreparedData<float>()[1]));
+		return true;
+	}
+
+	IMPLEMENT_NNE_SHAPEINFERENCEHELPER_UNIT_AUTOMATION_TEST(FElementWiseUnaryCPUHelperTanh, "System.Engine.MachineLearning.NNE.RDG.UnitTest.UnaryHelper.Tanh");
+	bool FElementWiseUnaryCPUHelperTanh::RunTest(const FString& Parameter)
+	{
+		FTensor XC2x1 = MakeConstTensor(TEXT("XC2x1"), { 2,1 }, { 1.0f, 2.0f });
+		FTensor Y = MakeTensor(TEXT("Y"), { 2,1 });
+		CPUHelper::ElementWiseUnary::Apply(EElementWiseUnaryOperatorType::Tanh, XC2x1, 1.0f, 2.0f, 3.0f, Y);
+		UTEST_TRUE(TEXT("Y const if input is const"), Y.HasPreparedData());
+		UTEST_EQUAL(TEXT("Tanh(XC2x1)[0]"), Y.GetPreparedData<float>()[0], 0.7615941559557649f);
+		UTEST_EQUAL(TEXT("Tanh(XC2x1)[1]"), Y.GetPreparedData<float>()[1], 0.9640275800758169f);
+		return true;
+	}
+
+#endif //WITH_DEV_AUTOMATION_TESTS
+
+} // namespace UE::NNERuntimeRDG::Private::Test::UnaryOp
diff --git a/Engine/Plugins/Experimental/NNERuntimeRDG/Source/NNERuntimeRDG/Private/Tests/NNERuntimeRDGUnitTestOperatorRegistry.cpp b/Engine/Plugins/Experimental/NNERuntimeRDG/Source/NNERuntimeRDG/Private/Tests/NNERuntimeRDGUnitTestOperatorRegistry.cpp
new file mode 100644
--- /dev/null
+++ b/Engine/Plugins/Experimental/NNERuntimeRDG/Source/NNERuntimeRDG/Private/Tests/NNERuntimeRDGUnitTestOperatorRegistry.cpp
@@ -0,0 +1,129 @@
+// Copyright Epic Games, Inc. All Rights Reserved.
+
+#include "Misc/AutomationTest.h"
+#include "NNERuntimeRDGBase.h"
+
+#if WITH_DEV_AUTOMATION_TESTS
+namespace UE::NNERuntimeRDG::Private::Test::Registry
+{
+
+// Base class for testing operator registry 
+class FOperatorTest
+{
+public:
+	virtual ~FOperatorTest() = default;
+};
+
+// Test operator that has single version
+class FOperatorTestConv : public FOperatorTest
+{
+public:
+	static FOperatorTest* Create()
+	{
+		return new FOperatorTestConv();
+	}
+	static bool Validate(const NNE::FAttributeMap& AttributeMap, TConstArrayView<ENNETensorDataType> InputTypes, TConstArrayView<NNE::FSymbolicTensorShape> InputShapes)
+	{
+		return true;
+	}
+};
+
+// Test operator that has multiple versions
+template<uint32 OperatorVersion>
+class FOperatorTestClip : public FOperatorTest
+{
+public:
+	static FOperatorTest* Create()
+	{
+		return new FOperatorTestClip();
+	}
+	static bool Validate(const NNE::FAttributeMap& AttributeMap, TConstArrayView<ENNETensorDataType> InputTypes, TConstArrayView<NNE::FSymbolicTensorShape> InputShapes)
+	{
+		return true;
+	}
+};
+
+IMPLEMENT_SIMPLE_AUTOMATION_TEST(
+	FOperatorRegistryRDGTest,
+	"System.Engine.MachineLearning.NNE.RDG.UnitTest.OperatorRegistry",
+	EAutomationTestFlags::EditorContext | EAutomationTestFlags::EngineFilter
+)
+bool FOperatorRegistryRDGTest::RunTest(const FString&)
+{
+    using FRegistryTest = TOperatorRegistryRDG<FOperatorTest>;
+
+    struct FOpRegistrationRef
+    {
+        FOperatorDesc           OpDesc;
+        FRegistryTest::OperatorCreateFunc 		CreateFunc;
+		FRegistryTest::OperatorValidateFunc 	ValidateFunc;
+    };
+
+    TArray<FOpRegistrationRef> TestOps;
+
+#define OP_TEST_CLASS(OpName) FOperatorTest##OpName
+#define OP_TEST_CLASS_TEMPL(OpName,OpVer) OP_TEST_CLASS(OpName) <OpVer>
+
+// Register operator with default version
+#define OP_REG(OpName) \
+{ \
+    FOperatorDesc OpDesc {{TEXT(#OpName), TEXT("Onnx")}}; \
+	FRegistryTest::Get()->OpAdd(OpDesc, OP_TEST_CLASS(OpName)::Create, OP_TEST_CLASS(OpName)::Validate); \
+	TestOps.Add({OpDesc, OP_TEST_CLASS(OpName)::Create, OP_TEST_CLASS(OpName)::Validate}); \
+} \
+const int Op##OpName##Idx = TestOps.Num() - 1;
+// Register operator with version
+#define OP_REG_VER(OpName, OpVer) \
+{ \
+    FOperatorDesc OpDesc {{TEXT(#OpName), TEXT("Onnx")}, OpVer}; \
+	FRegistryTest::Get()->OpAdd(OpDesc, OP_TEST_CLASS_TEMPL(OpName,OpVer)::Create, OP_TEST_CLASS_TEMPL(OpName,OpVer)::Validate); \
+	TestOps.Add({OpDesc, OP_TEST_CLASS_TEMPL(OpName,OpVer)::Create, OP_TEST_CLASS_TEMPL(OpName,OpVer)::Validate}); \
+} \
+const int Op##OpName##OpVer##Idx = TestOps.Num() - 1;
+
+	bool bPrevSuppressLogWarnings = bSuppressLogWarnings;
+    bSuppressLogWarnings = true;
+
+	OP_REG(Conv)
+	OP_REG_VER(Clip, 9)
+	OP_REG_VER(Clip, 11)
+	OP_REG_VER(Clip, 19)
+	
+	// Test registration
+	for (int32 Idx = 0; Idx < TestOps.Num(); ++Idx)
+	{
+		FRegistryTest::OperatorCreateFunc CreateFunc = FRegistryTest::Get()->OpFind(TestOps[Idx].OpDesc);
+		UTEST_NOT_NULL(TEXT("OpFind()"), (void *) CreateFunc);
+        UTEST_EQUAL(TEXT("OpFoundMatches"), CreateFunc, TestOps[Idx].CreateFunc);
+		FOperatorTest* Op = CreateFunc();
+		UTEST_NOT_NULL(TEXT("OpCreateCall"), Op);
+	}
+	// Test validation
+	for (int32 Idx = 0; Idx < TestOps.Num(); ++Idx)
+	{
+		FRegistryTest::OperatorValidateFunc ValidateFunc = FRegistryTest::Get()->OpFindValidation(TestOps[Idx].OpDesc);
+		UTEST_NOT_NULL(TEXT("OpFindValidation()"), (void *) ValidateFunc);
+        UTEST_EQUAL(TEXT("OpFoundMatches"), ValidateFunc, TestOps[Idx].ValidateFunc);
+        UE::NNE::FAttributeMap					Attributes;
+		TArray<ENNETensorDataType>				InputTypes;
+		TArray<UE::NNE::FSymbolicTensorShape>	InputShapes;
+		UTEST_TRUE(TEXT("OpValidateCall"), ValidateFunc(Attributes, InputTypes, InputShapes));
+	}
+
+    UTEST_EQUAL(TEXT("Find3InConvUnversioned"), 
+        FRegistryTest::Get()->OpFind({{TEXT("Conv"), TEXT("Onnx")}, 3}), 
+        TestOps[OpConvIdx].CreateFunc
+    );
+
+    
+    UTEST_NULL(TEXT("FindUnversionedInClipVersioned"), 
+        (void *) FRegistryTest::Get()->OpFind({{TEXT("Clip"), TEXT("Onnx")}})
+    );
+	
+    bSuppressLogWarnings = bPrevSuppressLogWarnings;
+
+	return true;
+}
+
+} // namespace UE::NNERuntimeRDG::Private::Test::Registry
+#endif
\ No newline at end of file
diff --git a/Engine/Plugins/Experimental/NNERuntimeRDG/Source/NNERuntimeRDG/Private/Tests/OperatorRegistryRDGTest.cpp b/Engine/Plugins/Experimental/NNERuntimeRDG/Source/NNERuntimeRDG/Private/Tests/OperatorRegistryRDGTest.cpp
deleted file mode 100644
diff --git a/Engine/Plugins/Experimental/NNERuntimeRDG/Source/NNERuntimeRDGUtils/Internal/NNERuntimeRDGUtilsModelOptimizer.h b/Engine/Plugins/Experimental/NNERuntimeRDG/Source/NNERuntimeRDGUtils/Internal/NNERuntimeRDGUtilsModelOptimizer.h
--- a/Engine/Plugins/Experimental/NNERuntimeRDG/Source/NNERuntimeRDGUtils/Internal/NNERuntimeRDGUtilsModelOptimizer.h
+++ b/Engine/Plugins/Experimental/NNERuntimeRDG/Source/NNERuntimeRDGUtils/Internal/NNERuntimeRDGUtilsModelOptimizer.h
@@ -2,8 +2,8 @@
 #pragma once
 
 #include "CoreMinimal.h"
-#include "NNEModelOptimizerInterface.h"
 #include "NNERuntimeFormat.h"
+#include "NNERuntimeRDGUtilsModelOptimizerInterface.h"
 
 namespace UE::NNE { class FAttributeMap; }
 
@@ -12,19 +12,19 @@ namespace UE::NNERuntimeRDGUtils::Internal
 
 using FOptimizerOptionsMap = NNE::FAttributeMap;
 	
-NNERUNTIMERDGUTILS_API TUniquePtr<NNE::Internal::IModelOptimizer> CreateModelOptimizer(ENNEInferenceFormat InputFormat, ENNEInferenceFormat OutputFormat);
+NNERUNTIMERDGUTILS_API TUniquePtr<IModelOptimizer> CreateModelOptimizer(ENNEInferenceFormat InputFormat, ENNEInferenceFormat OutputFormat);
 
-inline TUniquePtr<NNE::Internal::IModelOptimizer> CreateONNXToNNEModelOptimizer()
+inline TUniquePtr<IModelOptimizer> CreateONNXToNNEModelOptimizer()
 {
 	return CreateModelOptimizer(ENNEInferenceFormat::ONNX, ENNEInferenceFormat::NNERT);
 }
 
-inline TUniquePtr<NNE::Internal::IModelOptimizer> CreateONNXToORTModelOptimizer()
+inline TUniquePtr<IModelOptimizer> CreateONNXToORTModelOptimizer()
 {
 	return CreateModelOptimizer(ENNEInferenceFormat::ONNX, ENNEInferenceFormat::ORT);
 }
 
-inline TUniquePtr<NNE::Internal::IModelOptimizer> CreateONNXToONNXModelOptimizer()
+inline TUniquePtr<IModelOptimizer> CreateONNXToONNXModelOptimizer()
 {
 	return CreateModelOptimizer(ENNEInferenceFormat::ONNX, ENNEInferenceFormat::ONNX);
 }
diff --git a/Engine/Plugins/Experimental/NNERuntimeRDG/Source/NNERuntimeRDGUtils/Internal/NNERuntimeRDGUtilsModelOptimizerBase.h b/Engine/Plugins/Experimental/NNERuntimeRDG/Source/NNERuntimeRDGUtils/Internal/NNERuntimeRDGUtilsModelOptimizerBase.h
--- a/Engine/Plugins/Experimental/NNERuntimeRDG/Source/NNERuntimeRDGUtils/Internal/NNERuntimeRDGUtilsModelOptimizerBase.h
+++ b/Engine/Plugins/Experimental/NNERuntimeRDG/Source/NNERuntimeRDGUtils/Internal/NNERuntimeRDGUtilsModelOptimizerBase.h
@@ -2,34 +2,34 @@
 #pragma once
 
 #include "CoreMinimal.h"
-#include "NNEModelOptimizerInterface.h"
 #include "NNERuntimeRDGUtilsModelOptimizer.h"
+#include "NNERuntimeRDGUtilsModelOptimizerInterface.h"
 #include "Templates/SharedPointer.h"
 
 struct FNNEModelRaw;
 
 namespace UE::NNERuntimeRDGUtils::Internal
 {
-class FModelValidatorONNX : public NNE::Internal::IModelValidator
+class FModelValidatorONNX : public IModelValidator
 {
 public:
 	virtual FString GetName() const override;
 	virtual bool ValidateModel(const FNNEModelRaw& InputModel, const FOptimizerOptionsMap& Options) const override;
 };
 
-class FModelOptimizerBase : public NNE::Internal::IModelOptimizer
+class FModelOptimizerBase : public IModelOptimizer
 {
 protected:
 	FModelOptimizerBase() {}
-	TArray<TSharedPtr<NNE::Internal::IModelOptimizerPass>> OptimizationPasses;
-	TArray<TSharedPtr<NNE::Internal::IModelValidator>> Validators;
+	TArray<TSharedPtr<IModelOptimizerPass>> OptimizationPasses;
+	TArray<TSharedPtr<IModelValidator>> Validators;
 
 	bool IsModelValid(const FNNEModelRaw& ModelToValidate, const FOptimizerOptionsMap& Options);
 	bool ApplyAllPassesAndValidations(FNNEModelRaw& OptimizedModel, const FOptimizerOptionsMap& Options);
 
 public:
-	virtual void AddOptimizationPass(TSharedPtr<NNE::Internal::IModelOptimizerPass> ModelOptimizerPass) override;
-	virtual void AddValidator(TSharedPtr<NNE::Internal::IModelValidator> ModelValidator) override;
+	virtual void AddOptimizationPass(TSharedPtr<IModelOptimizerPass> ModelOptimizerPass) override;
+	virtual void AddValidator(TSharedPtr<IModelValidator> ModelValidator) override;
 	virtual bool Optimize(const FNNEModelRaw& InputModel, FNNEModelRaw& OptimizedModel, const FOptimizerOptionsMap& Options) override;
 };
 
diff --git a/Engine/Plugins/Experimental/NNERuntimeRDG/Source/NNERuntimeRDGUtils/Internal/NNERuntimeRDGUtilsModelOptimizerInterface.h b/Engine/Plugins/Experimental/NNERuntimeRDG/Source/NNERuntimeRDGUtils/Internal/NNERuntimeRDGUtilsModelOptimizerInterface.h
new file mode 100644
--- /dev/null
+++ b/Engine/Plugins/Experimental/NNERuntimeRDG/Source/NNERuntimeRDGUtils/Internal/NNERuntimeRDGUtilsModelOptimizerInterface.h
@@ -0,0 +1,55 @@
+// Copyright Epic Games, Inc. All Rights Reserved.
+
+#pragma once
+
+#include "Containers/UnrealString.h"
+#include "Templates/SharedPointerFwd.h"
+
+struct FNNEModelRaw;
+namespace UE::NNE { class FAttributeMap; }
+
+namespace UE::NNERuntimeRDGUtils::Internal
+{
+	
+//Note: create a dedicated FOptimizerOptionsMap when function diverge with FAttributeMap
+//example introduction of sparse tensor to FAttribute witch make no sense as an optimizer attribute
+using FOptimizerOptionsMap = UE::NNE::FAttributeMap;
+
+/** Interface class for NNE model validator */
+class IModelValidator
+{
+public:
+	virtual ~IModelValidator() = default;
+	virtual FString GetName() const = 0;
+	virtual bool ValidateModel(const FNNEModelRaw& InputModel, const FOptimizerOptionsMap& Options) const = 0;
+};
+
+/** Interface class for NNE model optimizer pass */
+class IModelOptimizerPass
+{
+public:
+	virtual ~IModelOptimizerPass() = default;
+	virtual FString GetName() const = 0;
+
+	//Optimize the model in place, potentially changing the format
+	virtual bool ApplyPass(FNNEModelRaw& Model, const FOptimizerOptionsMap& Options) const = 0;
+};
+
+/** Interface class for NNE model optimizer */
+class IModelOptimizer
+{
+public:
+	virtual ~IModelOptimizer() = default;
+	virtual FString GetName() const = 0;
+
+	//Allow to extend/customize an optimizer by adding passes. They should be executed in order.
+	virtual void AddOptimizationPass(TSharedPtr<IModelOptimizerPass> ModelOptimizerPass) = 0;
+	
+	//Allow to extend/customize an optimizer all validators should be run between each pass.
+	virtual void AddValidator(TSharedPtr<IModelValidator>) = 0;
+	
+	//Apply all passes and validators to the input model, produce an optimized model potentially in a different format
+	virtual bool Optimize(const FNNEModelRaw& InputModel, FNNEModelRaw& OutModel, const FOptimizerOptionsMap& Options) = 0;
+};
+
+} // namespace UE::NNERuntimeRDGUtils::Internal
diff --git a/Engine/Plugins/Experimental/NNERuntimeRDG/Source/NNERuntimeRDGUtils/Private/NNERuntimeRDGUtilsModelOptimizer.cpp b/Engine/Plugins/Experimental/NNERuntimeRDG/Source/NNERuntimeRDGUtils/Private/NNERuntimeRDGUtilsModelOptimizer.cpp
--- a/Engine/Plugins/Experimental/NNERuntimeRDG/Source/NNERuntimeRDGUtils/Private/NNERuntimeRDGUtilsModelOptimizer.cpp
+++ b/Engine/Plugins/Experimental/NNERuntimeRDG/Source/NNERuntimeRDGUtils/Private/NNERuntimeRDGUtilsModelOptimizer.cpp
@@ -8,7 +8,7 @@
 namespace UE::NNERuntimeRDGUtils::Internal
 {
 
-TUniquePtr<NNE::Internal::IModelOptimizer> CreateModelOptimizer(ENNEInferenceFormat InputFormat, ENNEInferenceFormat OutputFormat)
+TUniquePtr<IModelOptimizer> CreateModelOptimizer(ENNEInferenceFormat InputFormat, ENNEInferenceFormat OutputFormat)
 {
 	if (InputFormat == ENNEInferenceFormat::ONNX)
 	{
diff --git a/Engine/Plugins/Experimental/NNERuntimeRDG/Source/NNERuntimeRDGUtils/Private/NNERuntimeRDGUtilsModelOptimizerBase.cpp b/Engine/Plugins/Experimental/NNERuntimeRDG/Source/NNERuntimeRDGUtils/Private/NNERuntimeRDGUtilsModelOptimizerBase.cpp
--- a/Engine/Plugins/Experimental/NNERuntimeRDG/Source/NNERuntimeRDGUtils/Private/NNERuntimeRDGUtilsModelOptimizerBase.cpp
+++ b/Engine/Plugins/Experimental/NNERuntimeRDG/Source/NNERuntimeRDGUtils/Private/NNERuntimeRDGUtilsModelOptimizerBase.cpp
@@ -50,7 +50,7 @@ bool FModelOptimizerBase::IsModelValid(const FNNEModelRaw& ModelToValidate, cons
 {
 	bool bIsModelValid = true;
 
-	for (TSharedPtr<NNE::Internal::IModelValidator>& Validator : Validators)
+	for (TSharedPtr<IModelValidator>& Validator : Validators)
 	{
 		check(Validator.IsValid());
 		if (!Validator->ValidateModel(ModelToValidate, Options))
@@ -70,7 +70,7 @@ bool FModelOptimizerBase::ApplyAllPassesAndValidations(FNNEModelRaw& OptimizedMo
 		return false;
 	}
 		
-	for (TSharedPtr<NNE::Internal::IModelOptimizerPass>& Pass : OptimizationPasses)
+	for (TSharedPtr<IModelOptimizerPass>& Pass : OptimizationPasses)
 	{
 		check(Pass.IsValid());
 
@@ -96,7 +96,7 @@ bool FModelOptimizerBase::ApplyAllPassesAndValidations(FNNEModelRaw& OptimizedMo
 	return true;
 }
 
-void FModelOptimizerBase::AddOptimizationPass(TSharedPtr<NNE::Internal::IModelOptimizerPass> ModelOptimizerPass)
+void FModelOptimizerBase::AddOptimizationPass(TSharedPtr<IModelOptimizerPass> ModelOptimizerPass)
 {
 	if (ModelOptimizerPass.IsValid())
 	{
@@ -104,7 +104,7 @@ void FModelOptimizerBase::AddOptimizationPass(TSharedPtr<NNE::Internal::IModelOp
 	}
 }
 
-void FModelOptimizerBase::AddValidator(TSharedPtr<NNE::Internal::IModelValidator> ModelValidator)
+void FModelOptimizerBase::AddValidator(TSharedPtr<IModelValidator> ModelValidator)
 {
 	if (ModelValidator.IsValid())
 	{
diff --git a/Engine/Plugins/Experimental/NNERuntimeRDG/Source/NNERuntimeRDGUtils/Private/NNERuntimeRDGUtilsModelOptimizerONNX.cpp b/Engine/Plugins/Experimental/NNERuntimeRDG/Source/NNERuntimeRDGUtils/Private/NNERuntimeRDGUtilsModelOptimizerONNX.cpp
--- a/Engine/Plugins/Experimental/NNERuntimeRDG/Source/NNERuntimeRDGUtils/Private/NNERuntimeRDGUtilsModelOptimizerONNX.cpp
+++ b/Engine/Plugins/Experimental/NNERuntimeRDG/Source/NNERuntimeRDGUtils/Private/NNERuntimeRDGUtilsModelOptimizerONNX.cpp
@@ -18,7 +18,7 @@ THIRD_PARTY_INCLUDES_END
 namespace UE::NNERuntimeRDGUtils::Internal
 {
 
-class FOnnxRuntimeModelOptimizerPass : public NNE::Internal::IModelOptimizerPass
+class FOnnxRuntimeModelOptimizerPass : public IModelOptimizerPass
 {
 	ENNEInferenceFormat	TargetFormat;
 
@@ -144,7 +144,7 @@ private:
 
 };
 
-class FOnnxDomainCleanupModelOptimizerPass : public NNE::Internal::IModelOptimizerPass
+class FOnnxDomainCleanupModelOptimizerPass : public IModelOptimizerPass
 {
 public:
 	virtual FString GetName() const
@@ -202,7 +202,7 @@ public:
 
 };
 
-class FOnnxShapeInferenceModelOptimizerPass : public NNE::Internal::IModelOptimizerPass
+class FOnnxShapeInferenceModelOptimizerPass : public IModelOptimizerPass
 {
 public:
 	virtual FString GetName() const
@@ -214,7 +214,7 @@ public:
 	{
 		if (Model.Format != ENNEInferenceFormat::ONNX)
 		{
-			UE_LOG(LogNNERuntimeRDGHlsl, Warning, TEXT("%s is expecting a mNNEdel in ONNX format but received %u."), *(GetName()), Model.Format);
+			UE_LOG(LogNNERuntimeRDGHlsl, Warning, TEXT("%s is expecting a model in ONNX format but received %u."), *(GetName()), Model.Format);
 			return false;
 		}
 
diff --git a/Engine/Plugins/NNE/NNERuntimeORT/Source/NNERuntimeORT/Private/NNERuntimeORT.cpp b/Engine/Plugins/NNE/NNERuntimeORT/Source/NNERuntimeORT/Private/NNERuntimeORT.cpp
--- a/Engine/Plugins/NNE/NNERuntimeORT/Source/NNERuntimeORT/Private/NNERuntimeORT.cpp
+++ b/Engine/Plugins/NNE/NNERuntimeORT/Source/NNERuntimeORT/Private/NNERuntimeORT.cpp
@@ -6,7 +6,6 @@
 #include "NNE.h"
 #include "NNEAttributeMap.h"
 #include "NNEModelData.h"
-#include "NNEModelOptimizerInterface.h"
 #include "NNERuntimeORTModel.h"
 #include "NNERuntimeORTModelFormat.h"
 #include "NNERuntimeORTUtils.h"
diff --git a/Engine/Source/Runtime/NNE/Internal/NNEModelOptimizerInterface.h b/Engine/Source/Runtime/NNE/Internal/NNEModelOptimizerInterface.h
deleted file mode 100644
diff --git a/Engine/Source/Runtime/NNE/Internal/NNEOperator.h b/Engine/Source/Runtime/NNE/Internal/NNEOperator.h
deleted file mode 100644
diff --git a/Engine/Source/Runtime/NNE/Private/NNEModelData.cpp b/Engine/Source/Runtime/NNE/Private/NNEModelData.cpp
--- a/Engine/Source/Runtime/NNE/Private/NNEModelData.cpp
+++ b/Engine/Source/Runtime/NNE/Private/NNEModelData.cpp
@@ -3,9 +3,6 @@
 #include "NNEModelData.h"
 
 #include "NNE.h"
-#include "NNEAttributeMap.h"
-#include "NNEModelOptimizerInterface.h"
-#include "NNERuntimeFormat.h"
 #include "Serialization/CustomVersion.h"
 #include "UObject/AssetRegistryTagsContext.h"
 #include "UObject/WeakInterfacePtr.h"
