[NNE] NNE internal cleanup step3 (Code part)

Problem: NNE plugin are dependant on Source/NNE/Internal
Solution (step3): remove NNERuntimeFormat.h and implement in RDGHlsl and QA plugin instead.

Test:
Win64 editor debug : NNEQA all tests + style transfer
Win64 package debug : NNEQA all tests + style transfer
Console package dev : NNEQA smoke tests + style transfer

#rnx

diff --git a/Engine/Plugins/Experimental/NNERuntimeRDG/NNERuntimeRDG.uplugin b/Engine/Plugins/Experimental/NNERuntimeRDG/NNERuntimeRDG.uplugin
--- a/Engine/Plugins/Experimental/NNERuntimeRDG/NNERuntimeRDG.uplugin
+++ b/Engine/Plugins/Experimental/NNERuntimeRDG/NNERuntimeRDG.uplugin
@@ -26,6 +26,11 @@
 			"Type": "RuntimeAndProgram",
 			"LoadingPhase": "Default"
 		},
+		{
+			"Name": "NNERuntimeRDGData",
+			"Type": "RuntimeAndProgram",
+			"LoadingPhase": "Default"
+		},
 		{
 			"Name": "NNERuntimeRDGUtils",
 			"Type": "EditorAndProgram",
diff --git a/Engine/Plugins/Experimental/NNERuntimeRDG/Source/NNERuntimeRDG/NNERuntimeRDG.Build.cs b/Engine/Plugins/Experimental/NNERuntimeRDG/Source/NNERuntimeRDG/NNERuntimeRDG.Build.cs
--- a/Engine/Plugins/Experimental/NNERuntimeRDG/Source/NNERuntimeRDG/NNERuntimeRDG.Build.cs
+++ b/Engine/Plugins/Experimental/NNERuntimeRDG/Source/NNERuntimeRDG/NNERuntimeRDG.Build.cs
@@ -23,6 +23,7 @@ public class NNERuntimeRDG : ModuleRules
 		{
 			"NNE",
 			"NNEHlslShaders",
+			"NNERuntimeRDGData",
 			"RHI",
 			"Projects",
 			"TraceLog"
diff --git a/Engine/Plugins/Experimental/NNERuntimeRDG/Source/NNERuntimeRDG/Private/NNERuntimeRDGBase.cpp b/Engine/Plugins/Experimental/NNERuntimeRDG/Source/NNERuntimeRDG/Private/NNERuntimeRDGBase.cpp
--- a/Engine/Plugins/Experimental/NNERuntimeRDG/Source/NNERuntimeRDG/Private/NNERuntimeRDGBase.cpp
+++ b/Engine/Plugins/Experimental/NNERuntimeRDG/Source/NNERuntimeRDG/Private/NNERuntimeRDGBase.cpp
@@ -3,7 +3,6 @@
 #include "NNERuntimeRDGBase.h"
 
 #include "Helper/NNERuntimeRDGLogHelper.h"
-#include "NNERuntimeFormat.h"
 #include "RenderGraphBuilder.h"
 #include "RenderGraphUtils.h"
 #include "RHIGPUReadback.h"
diff --git a/Engine/Plugins/Experimental/NNERuntimeRDG/Source/NNERuntimeRDG/Private/NNERuntimeRDGBase.h b/Engine/Plugins/Experimental/NNERuntimeRDG/Source/NNERuntimeRDG/Private/NNERuntimeRDGBase.h
--- a/Engine/Plugins/Experimental/NNERuntimeRDG/Source/NNERuntimeRDG/Private/NNERuntimeRDGBase.h
+++ b/Engine/Plugins/Experimental/NNERuntimeRDG/Source/NNERuntimeRDG/Private/NNERuntimeRDGBase.h
@@ -2,10 +2,12 @@
 
 #pragma once
 
+#include "CoreMinimal.h"
+#include "Containers/ContainersFwd.h"
 #include "NNEAttributeMap.h"
 #include "NNEAttributeValue.h"
 #include "NNEHlslShadersLog.h"
-#include "NNERuntimeFormat.h"
+#include "NNERuntimeRDGDataFormat.h"
 #ifdef NNE_UTILITIES_AVAILABLE
 #include "NNERuntimeRDGUtilsModelOptimizerInterface.h"
 #endif //NNE_UTILITIES_AVAILABLE
@@ -22,7 +24,6 @@ BEGIN_SHADER_PARAMETER_STRUCT(FNNETensorReadbackParameters, )
 END_SHADER_PARAMETER_STRUCT()
 
 class FRDGBuilder;
-struct FNNEModelRaw;
 
 namespace UE::NNE { class FTensorDesc; }
 
@@ -354,19 +355,11 @@ public:
 		return TEXT("RDG Model validator");
 	}
 	
-	virtual bool ValidateModel(const FNNEModelRaw& InputModel, const NNERuntimeRDGUtils::Internal::FOptimizerOptionsMap& Options) const override
+	virtual bool ValidateModel(TConstArrayView<uint8> InputModel) const override
 	{
-		FNNERuntimeFormat	Format;
+		FNNERuntimeRDGDataModelFormat	Format;
 
-		ENNEInferenceFormat FormatType = InputModel.Format;
-		if (FormatType != ENNEInferenceFormat::NNERT)
-		{
-			UE_LOG(LogNNERuntimeRDGHlsl, Warning, TEXT("Unsupported format type for validator %s"), *GetName());
-			return false;
-		}
-
-		FMemoryReader Reader(InputModel.Data, /*bIsPersitent =*/ true);
-		Reader.SetIsPersistent(true);
+		FMemoryReaderView Reader(InputModel, /*bIsPersitent =*/ true);
 
 		Format.Serialize(Reader);
 
@@ -384,7 +377,7 @@ public:
 				InputTensorTypes.Add(Format.Tensors[InputTensorIndex].DataType);
 				InputTensorShapes.Add(NNE::FSymbolicTensorShape::Make(Format.Tensors[InputTensorIndex].Shape));
 			}
-			for (const FNNEFormatAttributeDesc& Desc : Format.Operators[Idx].Attributes)
+			for (const FNNERuntimeRDGDataAttributeDesc& Desc : Format.Operators[Idx].Attributes)
 			{
 				AttributeMap.SetAttribute(Desc.Name, Desc.Value);
 			}
diff --git a/Engine/Plugins/Experimental/NNERuntimeRDG/Source/NNERuntimeRDG/Private/NNERuntimeRDGHlsl.cpp b/Engine/Plugins/Experimental/NNERuntimeRDG/Source/NNERuntimeRDG/Private/NNERuntimeRDGHlsl.cpp
--- a/Engine/Plugins/Experimental/NNERuntimeRDG/Source/NNERuntimeRDG/Private/NNERuntimeRDGHlsl.cpp
+++ b/Engine/Plugins/Experimental/NNERuntimeRDG/Source/NNERuntimeRDG/Private/NNERuntimeRDGHlsl.cpp
@@ -231,25 +231,21 @@ TSharedPtr<UE::NNE::FSharedModelData> UNNERuntimeRDGHlslImpl::CreateModelData(co
 	}
 
 #ifdef NNE_UTILITIES_AVAILABLE
-	TUniquePtr<UE::NNERuntimeRDGUtils::Internal::IModelOptimizer> Optimizer = UE::NNERuntimeRDGUtils::Internal::CreateONNXToNNEModelOptimizer();
+	TUniquePtr<UE::NNERuntimeRDGUtils::Internal::IModelOptimizer> Optimizer = UE::NNERuntimeRDGUtils::Internal::CreateModelOptimizer();
 	Optimizer->AddValidator(MakeShared<UE::NNERuntimeRDG::Private::TModelValidatorRDG<FOperatorHlsl>>());
 
-	FNNEModelRaw InputModel;
-	InputModel.Data = FileData;
-	InputModel.Format = ENNEInferenceFormat::ONNX;
-
-	FNNEModelRaw OutputModel;
-	if (!Optimizer->Optimize(InputModel, OutputModel, {}))
+	TArray<uint8> OutputModel;
+	if (!Optimizer->Optimize(FileData, OutputModel))
 	{
 		return {};
 	}
 
-	TArray64<uint8> Result;
-	FMemoryWriter64 Writer(Result);
+	TArray<uint8> Result;
+	FMemoryWriter Writer(Result);
 	
 	Writer << GUID;
 	Writer << Version;
-	Writer.Serialize(OutputModel.Data.GetData(), OutputModel.Data.Num());
+	Writer.Serialize(OutputModel.GetData(), OutputModel.Num());
 
 	return MakeShared<UE::NNE::FSharedModelData>(MakeSharedBufferFromArray(MoveTemp(Result)), 0);
 #else //NNE_UTILITIES_AVAILABLE
diff --git a/Engine/Plugins/Experimental/NNERuntimeRDG/Source/NNERuntimeRDG/Private/NNERuntimeRDGModel.cpp b/Engine/Plugins/Experimental/NNERuntimeRDG/Source/NNERuntimeRDG/Private/NNERuntimeRDGModel.cpp
--- a/Engine/Plugins/Experimental/NNERuntimeRDG/Source/NNERuntimeRDG/Private/NNERuntimeRDGModel.cpp
+++ b/Engine/Plugins/Experimental/NNERuntimeRDG/Source/NNERuntimeRDG/Private/NNERuntimeRDGModel.cpp
@@ -3,7 +3,7 @@
 #include "NNERuntimeRDGModel.h"
 
 #include "NNEHlslShadersLog.h"
-#include "NNERuntimeFormat.h"
+#include "NNERuntimeRDGDataFormat.h"
 #include "RenderGraphBuilder.h"
 #include "RenderGraphUtils.h"
 #include "RHIGPUReadback.h"
@@ -13,7 +13,7 @@
 namespace UE::NNERuntimeRDG::Private
 {
 
-bool FModelInstanceRDG::LoadModel(TConstArrayView<uint8> ModelData, FNNERuntimeFormat& Format, int32 GuidAndVersionSize)
+bool FModelInstanceRDG::LoadModel(TConstArrayView<uint8> ModelData, FNNERuntimeRDGDataModelFormat& Format, int32 GuidAndVersionSize)
 {
 	TConstArrayView<uint8> ModelBuffer = { &(ModelData.GetData()[GuidAndVersionSize]), ModelData.Num() - GuidAndVersionSize };
 
@@ -39,12 +39,12 @@ bool FModelInstanceRDG::LoadModel(TConstArrayView<uint8> ModelData, FNNERuntimeF
 	// Add tensors
 	for (int32 Idx = 0; Idx < Format.Tensors.Num(); ++Idx)
 	{
-		const FNNEFormatTensorDesc& FormatTensorDesc = Format.Tensors[Idx];
+		const FNNERuntimeRDGDataTensorDesc& FormatTensorDesc = Format.Tensors[Idx];
 
 		const NNE::FSymbolicTensorShape SymbolicShape = NNE::FSymbolicTensorShape::Make(FormatTensorDesc.Shape);
 		const NNE::FTensorDesc SymbolicTensor = NNE::FTensorDesc::Make(FormatTensorDesc.Name, SymbolicShape, FormatTensorDesc.DataType);
 
-		if (FormatTensorDesc.Type != ENNEFormatTensorType::Empty && FormatTensorDesc.DataType == ENNETensorDataType::None)
+		if (FormatTensorDesc.Type != ENNERuntimeRDGDataTensorType::Empty && FormatTensorDesc.DataType == ENNETensorDataType::None)
 		{
 			UE_LOG(LogNNERuntimeRDGHlsl, Error, TEXT("Tensor %s has invalid format: Data type None is reserved for empty tensors."), *SymbolicTensor.GetName());
 			return false;
@@ -52,25 +52,25 @@ bool FModelInstanceRDG::LoadModel(TConstArrayView<uint8> ModelData, FNNERuntimeF
 
 		AllSymbolicTensorDescs.Emplace(Idx, SymbolicTensor);
 
-		if (FormatTensorDesc.Type == ENNEFormatTensorType::Empty)
+		if (FormatTensorDesc.Type == ENNERuntimeRDGDataTensorType::Empty)
 		{
 			EmptyTensorIndices.Emplace(Idx);
 		}
-		else if (FormatTensorDesc.Type == ENNEFormatTensorType::Input)
+		else if (FormatTensorDesc.Type == ENNERuntimeRDGDataTensorType::Input)
 		{
 			InputTensorIndices.Emplace(Idx);
 			InputSymbolicTensors.Emplace(SymbolicTensor);
 		}
-		else if (FormatTensorDesc.Type == ENNEFormatTensorType::Output)
+		else if (FormatTensorDesc.Type == ENNERuntimeRDGDataTensorType::Output)
 		{
 			OutputTensorIndices.Emplace(Idx);
 			OutputSymbolicTensors.Emplace(SymbolicTensor);
 		}
-		else if (FormatTensorDesc.Type == ENNEFormatTensorType::Intermediate)
+		else if (FormatTensorDesc.Type == ENNERuntimeRDGDataTensorType::Intermediate)
 		{
 			IntermediateTensorIndices.Emplace(Idx);
 		}
-		else if (FormatTensorDesc.Type == ENNEFormatTensorType::Initializer)
+		else if (FormatTensorDesc.Type == ENNERuntimeRDGDataTensorType::Initializer)
 		{
 			WeightTensorIndices.Emplace(Idx);
 			if (!SymbolicTensor.GetShape().IsConcrete())
@@ -93,7 +93,7 @@ bool FModelInstanceRDG::LoadModel(TConstArrayView<uint8> ModelData, FNNERuntimeF
 
 			WeightRDG.SetPreparedData(DataView);
 		}
-		checkf(FormatTensorDesc.Type != ENNEFormatTensorType::None, TEXT("Unsupported tensor type None"));
+		checkf(FormatTensorDesc.Type != ENNERuntimeRDGDataTensorType::None, TEXT("Unsupported tensor type None"));
 	}
 
 	// Loop over all operators in the model and store tensor indices for input/output
diff --git a/Engine/Plugins/Experimental/NNERuntimeRDG/Source/NNERuntimeRDG/Private/NNERuntimeRDGModel.h b/Engine/Plugins/Experimental/NNERuntimeRDG/Source/NNERuntimeRDG/Private/NNERuntimeRDGModel.h
--- a/Engine/Plugins/Experimental/NNERuntimeRDG/Source/NNERuntimeRDG/Private/NNERuntimeRDGModel.h
+++ b/Engine/Plugins/Experimental/NNERuntimeRDG/Source/NNERuntimeRDG/Private/NNERuntimeRDGModel.h
@@ -28,7 +28,7 @@ namespace UE::NNERuntimeRDG::Private
 		
 	protected:
 		static FRDGBufferDesc CreateRDGBufferDescForTensorRDG(const FTensorRDG& Tensor);
-		bool LoadModel(TConstArrayView<uint8> ModelData, FNNERuntimeFormat& Format, int32 GuidAndVersionSize);
+		bool LoadModel(TConstArrayView<uint8> ModelData, FNNERuntimeRDGDataModelFormat& Format, int32 GuidAndVersionSize);
 		int32 SetTensors(FRDGBuilder& GraphBuilder, FTensorRDGArray& InTensorRDGs, TConstArrayView<NNE::FTensorBindingRDG> InBindings);
 
 		virtual int PrepareTensorShapesAndData() = 0;
diff --git a/Engine/Plugins/Experimental/NNERuntimeRDG/Source/NNERuntimeRDG/Private/NNERuntimeRDGModelHlsl.cpp b/Engine/Plugins/Experimental/NNERuntimeRDG/Source/NNERuntimeRDG/Private/NNERuntimeRDGModelHlsl.cpp
--- a/Engine/Plugins/Experimental/NNERuntimeRDG/Source/NNERuntimeRDG/Private/NNERuntimeRDGModelHlsl.cpp
+++ b/Engine/Plugins/Experimental/NNERuntimeRDG/Source/NNERuntimeRDG/Private/NNERuntimeRDGModelHlsl.cpp
@@ -38,7 +38,7 @@ FOperatorHlsl* OpCreate(const FOperatorDesc& OpDesc, TConstArrayView<NNE::FTenso
 }
 
 FTensorRDGRef GetTensorRefIfWeightAndOnlyUsedByOperator(int32 TensorIndex, int32 OperatorIdx,
-	TConstArrayView<FNNEFormatOperatorDesc> Operators, TConstArrayView<int32> WeightTensorIndices, FTensorRDGArray& WeightTensorRDGs)
+	TConstArrayView<FNNERuntimeRDGDataOperatorDesc> Operators, TConstArrayView<int32> WeightTensorIndices, FTensorRDGArray& WeightTensorRDGs)
 {
 	for (int32 Idx = 0; Idx < Operators.Num(); ++Idx)
 	{
@@ -126,7 +126,7 @@ FModelInstance::~FModelInstance()
 bool FModelInstance::Init(TConstArrayView<uint8> ModelData)
 {
 	check(ModelData.Num() > 0);
-	FNNERuntimeFormat	Format;
+	FNNERuntimeRDGDataModelFormat Format;
 	int32 GuidSize = sizeof(UNNERuntimeRDGHlslImpl::GUID);
 	int32 VersionSize = sizeof(UNNERuntimeRDGHlslImpl::Version);
 	
@@ -159,7 +159,7 @@ bool FModelInstance::Init(TConstArrayView<uint8> ModelData)
 		{
 			Outputs.Emplace(AllSymbolicTensorDescs[OutputTensorIndex]);
 		}
-		for (const FNNEFormatAttributeDesc& Desc : Format.Operators[Idx].Attributes)
+		for (const FNNERuntimeRDGDataAttributeDesc& Desc : Format.Operators[Idx].Attributes)
 		{
 			AttributeMap.SetAttribute(Desc.Name, Desc.Value);
 		}
diff --git a/Engine/Plugins/Experimental/NNERuntimeRDG/Source/NNERuntimeRDGData/Internal/NNERuntimeRDGDataFormat.h b/Engine/Plugins/Experimental/NNERuntimeRDG/Source/NNERuntimeRDGData/Internal/NNERuntimeRDGDataFormat.h
new file mode 100644
--- /dev/null
+++ b/Engine/Plugins/Experimental/NNERuntimeRDG/Source/NNERuntimeRDGData/Internal/NNERuntimeRDGDataFormat.h
@@ -0,0 +1,129 @@
+// Copyright Epic Games, Inc. All Rights Reserved.
+
+#pragma once
+
+#include "NNEAttributeValue.h"
+#include "NNETypes.h"
+#include "UObject/Object.h"
+
+#include "NNERuntimeRDGDataFormat.generated.h"
+
+UENUM()
+enum class ENNERuntimeRDGDataTensorType : uint8
+{
+	None,	
+	Input,
+	Output,
+	Intermediate,
+	Initializer,
+	Empty,
+
+	NUM
+};
+
+USTRUCT()
+struct FNNERuntimeRDGDataAttributeDesc
+{
+	GENERATED_USTRUCT_BODY()
+
+	UPROPERTY(VisibleAnywhere, Category = "Neural Network Inference")
+	FString Name;
+
+	UPROPERTY(VisibleAnywhere, Category = "Neural Network Inference")
+	FNNEAttributeValue Value;
+};
+
+USTRUCT()
+struct FNNERuntimeRDGDataOperatorDesc
+{
+	GENERATED_USTRUCT_BODY()
+
+	UPROPERTY(VisibleAnywhere, Category = "Neural Network Inference")
+	FString TypeName;			//!< For example "Relu"
+
+	UPROPERTY(VisibleAnywhere, Category = "Neural Network Inference")
+	FString DomainName;			//!< For example "onnx"
+
+	UPROPERTY(VisibleAnywhere, Category = "Neural Network Inference")
+	TOptional<uint32> Version;	//!< For example 7
+
+	UPROPERTY(VisibleAnywhere, Category = "Neural Network Inference")
+	TArray<uint32> InTensors;
+
+	UPROPERTY(VisibleAnywhere, Category = "Neural Network Inference")
+	TArray<uint32> OutTensors;
+
+	UPROPERTY(VisibleAnywhere, Category = "Neural Network Inference")
+	TArray<FNNERuntimeRDGDataAttributeDesc> Attributes;
+};
+
+USTRUCT()
+struct FNNERuntimeRDGDataTensorDesc
+{
+	GENERATED_USTRUCT_BODY()
+
+	UPROPERTY(VisibleAnywhere, Category = "Neural Network Inference")
+	FString Name;
+
+	UPROPERTY(VisibleAnywhere, Category = "Neural Network Inference")
+	TArray<int32> Shape;
+
+	UPROPERTY(VisibleAnywhere, Category = "Neural Network Inference")
+	ENNERuntimeRDGDataTensorType	Type = ENNERuntimeRDGDataTensorType::None;
+
+	UPROPERTY(VisibleAnywhere, Category = "Neural Network Inference")
+	ENNETensorDataType	DataType = ENNETensorDataType::None;
+
+	UPROPERTY(VisibleAnywhere, Category = "Neural Network Inference")
+	uint64	DataSize = 0;
+
+	UPROPERTY(VisibleAnywhere, Category = "Neural Network Inference")
+	uint64	DataOffset = 0;
+};
+
+USTRUCT()
+struct FNNERuntimeRDGDataModelFormat
+{
+	GENERATED_USTRUCT_BODY()
+
+	UPROPERTY(VisibleAnywhere, Category = "Neural Network Inference")
+	TArray<FNNERuntimeRDGDataTensorDesc> Tensors;
+
+	UPROPERTY(VisibleAnywhere, Category = "Neural Network Inference")
+	TArray<FNNERuntimeRDGDataOperatorDesc> Operators;
+
+	uint64 DataSize;
+	TArray64<uint8> TensorData;
+	
+	bool Serialize(FArchive& Ar)
+	{
+		// Serialize normal UPROPERTY tagged data
+		UScriptStruct* Struct = FNNERuntimeRDGDataModelFormat::StaticStruct();
+		Struct->SerializeTaggedProperties(Ar, reinterpret_cast<uint8*>(this), Struct, nullptr);
+
+		if (Ar.IsLoading())
+		{
+			Ar << DataSize;
+			TensorData.SetNumUninitialized(DataSize);
+			Ar.Serialize((void*)TensorData.GetData(), DataSize);
+		}
+		else if (Ar.IsSaving())
+		{
+			DataSize = TensorData.Num();
+			Ar << DataSize;
+			Ar.Serialize((void*)TensorData.GetData(), DataSize);
+		}
+
+		return true;
+	}
+
+};
+
+template<>
+struct TStructOpsTypeTraits<FNNERuntimeRDGDataModelFormat> : public TStructOpsTypeTraitsBase2<FNNERuntimeRDGDataModelFormat>
+{
+	enum
+	{
+		WithSerializer = true
+	};
+};
\ No newline at end of file
diff --git a/Engine/Plugins/Experimental/NNERuntimeRDG/Source/NNERuntimeRDGData/NNERuntimeRDGData.Build.cs b/Engine/Plugins/Experimental/NNERuntimeRDG/Source/NNERuntimeRDGData/NNERuntimeRDGData.Build.cs
new file mode 100644
--- /dev/null
+++ b/Engine/Plugins/Experimental/NNERuntimeRDG/Source/NNERuntimeRDGData/NNERuntimeRDGData.Build.cs
@@ -0,0 +1,24 @@
+// Copyright Epic Games, Inc. All Rights Reserved.
+
+using UnrealBuildTool;
+using System.IO;
+
+public class NNERuntimeRDGData : ModuleRules
+{
+	public NNERuntimeRDGData(ReadOnlyTargetRules Target) : base(Target)
+	{
+		PCHUsage = ModuleRules.PCHUsageMode.UseExplicitOrSharedPCHs;
+
+		PublicDependencyModuleNames.Add("NNE");
+
+		PrivateDependencyModuleNames.AddRange(
+			new string[]
+			{
+				"Core",
+				"CoreUObject",
+				"Engine"
+			}
+		);
+	}
+}
+
diff --git a/Engine/Plugins/Experimental/NNERuntimeRDG/Source/NNERuntimeRDGData/Private/NNERuntimeRDGDataModule.cpp b/Engine/Plugins/Experimental/NNERuntimeRDG/Source/NNERuntimeRDGData/Private/NNERuntimeRDGDataModule.cpp
new file mode 100644
--- /dev/null
+++ b/Engine/Plugins/Experimental/NNERuntimeRDG/Source/NNERuntimeRDGData/Private/NNERuntimeRDGDataModule.cpp
@@ -0,0 +1,19 @@
+// Copyright Epic Games, Inc. All Rights Reserved.
+
+#include "Modules/ModuleInterface.h"
+#include "Modules/ModuleManager.h"
+
+class FNNERuntimeRDGDataModule : public IModuleInterface
+{
+public:
+
+	virtual void StartupModule() override
+	{
+	}
+
+	virtual void ShutdownModule() override
+	{
+	}
+};
+
+IMPLEMENT_MODULE(FNNERuntimeRDGDataModule, NNERuntimeRDGData)
\ No newline at end of file
diff --git a/Engine/Plugins/Experimental/NNERuntimeRDG/Source/NNERuntimeRDGUtils/Internal/NNERuntimeRDGUtilsModelBuilderNNE.h b/Engine/Plugins/Experimental/NNERuntimeRDG/Source/NNERuntimeRDGUtils/Internal/NNERuntimeRDGUtilsModelBuilderNNE.h
deleted file mode 100644
diff --git a/Engine/Plugins/Experimental/NNERuntimeRDG/Source/NNERuntimeRDGUtils/Internal/NNERuntimeRDGUtilsModelBuilderONNX.h b/Engine/Plugins/Experimental/NNERuntimeRDG/Source/NNERuntimeRDGUtils/Internal/NNERuntimeRDGUtilsModelBuilderONNX.h
--- a/Engine/Plugins/Experimental/NNERuntimeRDG/Source/NNERuntimeRDGUtils/Internal/NNERuntimeRDGUtilsModelBuilderONNX.h
+++ b/Engine/Plugins/Experimental/NNERuntimeRDG/Source/NNERuntimeRDGUtils/Internal/NNERuntimeRDGUtilsModelBuilderONNX.h
@@ -9,8 +9,6 @@
 #include "Templates/UniquePtr.h"
 
 namespace UE::NNE { class FAttributeMap; }
-struct FNNEAttributeValue;
-struct FNNEModelRaw;
 
 namespace UE::NNERuntimeRDGUtils::Internal
 {
@@ -20,7 +18,7 @@ NNERUNTIMERDGUTILS_API TUniquePtr<IModelBuilder> CreateONNXModelBuilder(int64 Ir
 NNERUNTIMERDGUTILS_API bool CreateONNXModelForOperator(const FString& OperatorName, int32 IrVersion, int32 OpsetVersion, bool bUseVariadicShapeForModel,
 	TConstArrayView<NNE::Internal::FTensor> InInputTensors, TConstArrayView<NNE::Internal::FTensor> InOutputTensors,
 	TConstArrayView<NNE::Internal::FTensor> InWeightTensors, TConstArrayView<TConstArrayView<uint8>> InWeightTensorsData,
-	const UE::NNE::FAttributeMap& Attributes, FNNEModelRaw& ModelData);
+	const UE::NNE::FAttributeMap& Attributes, TArray<uint8>& ModelData);
 
 } // UE::NNERuntimeRDGUtils::Internal
 
diff --git a/Engine/Plugins/Experimental/NNERuntimeRDG/Source/NNERuntimeRDGUtils/Internal/NNERuntimeRDGUtilsModelOptimizer.h b/Engine/Plugins/Experimental/NNERuntimeRDG/Source/NNERuntimeRDGUtils/Internal/NNERuntimeRDGUtilsModelOptimizer.h
--- a/Engine/Plugins/Experimental/NNERuntimeRDG/Source/NNERuntimeRDGUtils/Internal/NNERuntimeRDGUtilsModelOptimizer.h
+++ b/Engine/Plugins/Experimental/NNERuntimeRDG/Source/NNERuntimeRDGUtils/Internal/NNERuntimeRDGUtilsModelOptimizer.h
@@ -2,31 +2,11 @@
 #pragma once
 
 #include "CoreMinimal.h"
-#include "NNERuntimeFormat.h"
 #include "NNERuntimeRDGUtilsModelOptimizerInterface.h"
 
-namespace UE::NNE { class FAttributeMap; }
-
 namespace UE::NNERuntimeRDGUtils::Internal
 {
 
-using FOptimizerOptionsMap = NNE::FAttributeMap;
-	
-NNERUNTIMERDGUTILS_API TUniquePtr<IModelOptimizer> CreateModelOptimizer(ENNEInferenceFormat InputFormat, ENNEInferenceFormat OutputFormat);
-
-inline TUniquePtr<IModelOptimizer> CreateONNXToNNEModelOptimizer()
-{
-	return CreateModelOptimizer(ENNEInferenceFormat::ONNX, ENNEInferenceFormat::NNERT);
-}
-
-inline TUniquePtr<IModelOptimizer> CreateONNXToORTModelOptimizer()
-{
-	return CreateModelOptimizer(ENNEInferenceFormat::ONNX, ENNEInferenceFormat::ORT);
-}
-
-inline TUniquePtr<IModelOptimizer> CreateONNXToONNXModelOptimizer()
-{
-	return CreateModelOptimizer(ENNEInferenceFormat::ONNX, ENNEInferenceFormat::ONNX);
-}
+NNERUNTIMERDGUTILS_API TUniquePtr<IModelOptimizer> CreateModelOptimizer();
 
 } // UE::NNERuntimeRDGUtils::Internal
diff --git a/Engine/Plugins/Experimental/NNERuntimeRDG/Source/NNERuntimeRDGUtils/Internal/NNERuntimeRDGUtilsModelOptimizerBase.h b/Engine/Plugins/Experimental/NNERuntimeRDG/Source/NNERuntimeRDGUtils/Internal/NNERuntimeRDGUtilsModelOptimizerBase.h
deleted file mode 100644
diff --git a/Engine/Plugins/Experimental/NNERuntimeRDG/Source/NNERuntimeRDGUtils/Internal/NNERuntimeRDGUtilsModelOptimizerInterface.h b/Engine/Plugins/Experimental/NNERuntimeRDG/Source/NNERuntimeRDGUtils/Internal/NNERuntimeRDGUtilsModelOptimizerInterface.h
--- a/Engine/Plugins/Experimental/NNERuntimeRDG/Source/NNERuntimeRDGUtils/Internal/NNERuntimeRDGUtilsModelOptimizerInterface.h
+++ b/Engine/Plugins/Experimental/NNERuntimeRDG/Source/NNERuntimeRDGUtils/Internal/NNERuntimeRDGUtilsModelOptimizerInterface.h
@@ -2,26 +2,20 @@
 
 #pragma once
 
+#include "Containers/ContainersFwd.h"
 #include "Containers/UnrealString.h"
 #include "Templates/SharedPointerFwd.h"
 
-struct FNNEModelRaw;
-namespace UE::NNE { class FAttributeMap; }
-
 namespace UE::NNERuntimeRDGUtils::Internal
 {
 	
-//Note: create a dedicated FOptimizerOptionsMap when function diverge with FAttributeMap
-//example introduction of sparse tensor to FAttribute witch make no sense as an optimizer attribute
-using FOptimizerOptionsMap = UE::NNE::FAttributeMap;
-
 /** Interface class for NNE model validator */
 class IModelValidator
 {
 public:
 	virtual ~IModelValidator() = default;
 	virtual FString GetName() const = 0;
-	virtual bool ValidateModel(const FNNEModelRaw& InputModel, const FOptimizerOptionsMap& Options) const = 0;
+	virtual bool ValidateModel(TConstArrayView<uint8> InputModel) const = 0;
 };
 
 /** Interface class for NNE model optimizer pass */
@@ -32,7 +26,7 @@ public:
 	virtual FString GetName() const = 0;
 
 	//Optimize the model in place, potentially changing the format
-	virtual bool ApplyPass(FNNEModelRaw& Model, const FOptimizerOptionsMap& Options) const = 0;
+	virtual bool ApplyPass(TArray<uint8>& Model) const = 0;
 };
 
 /** Interface class for NNE model optimizer */
@@ -49,7 +43,7 @@ public:
 	virtual void AddValidator(TSharedPtr<IModelValidator>) = 0;
 	
 	//Apply all passes and validators to the input model, produce an optimized model potentially in a different format
-	virtual bool Optimize(const FNNEModelRaw& InputModel, FNNEModelRaw& OutModel, const FOptimizerOptionsMap& Options) = 0;
+	virtual bool Optimize(TConstArrayView<uint8> InputModel, TArray<uint8>& OutModel) = 0;
 };
 
 } // namespace UE::NNERuntimeRDGUtils::Internal
diff --git a/Engine/Plugins/Experimental/NNERuntimeRDG/Source/NNERuntimeRDGUtils/Internal/NNERuntimeRDGUtilsModelOptimizerNNE.h b/Engine/Plugins/Experimental/NNERuntimeRDG/Source/NNERuntimeRDGUtils/Internal/NNERuntimeRDGUtilsModelOptimizerNNE.h
deleted file mode 100644
diff --git a/Engine/Plugins/Experimental/NNERuntimeRDG/Source/NNERuntimeRDGUtils/Internal/NNERuntimeRDGUtilsModelOptimizerONNX.h b/Engine/Plugins/Experimental/NNERuntimeRDG/Source/NNERuntimeRDGUtils/Internal/NNERuntimeRDGUtilsModelOptimizerONNX.h
deleted file mode 100644
diff --git a/Engine/Plugins/Experimental/NNERuntimeRDG/Source/NNERuntimeRDGUtils/NNERuntimeRDGUtils.Build.cs b/Engine/Plugins/Experimental/NNERuntimeRDG/Source/NNERuntimeRDGUtils/NNERuntimeRDGUtils.Build.cs
--- a/Engine/Plugins/Experimental/NNERuntimeRDG/Source/NNERuntimeRDGUtils/NNERuntimeRDGUtils.Build.cs
+++ b/Engine/Plugins/Experimental/NNERuntimeRDG/Source/NNERuntimeRDGUtils/NNERuntimeRDGUtils.Build.cs
@@ -20,6 +20,7 @@ public class NNERuntimeRDGUtils : ModuleRules
 				"Engine",
 				"Projects",
 				"NNEHlslShaders",
+				"NNERuntimeRDGData",
 				"NNERuntimeRDGProtobufEditor",
 				"NNERuntimeRDGOnnxruntimeEditor",
 				"NNERuntimeRDGOnnxEditor"
diff --git a/Engine/Plugins/Experimental/NNERuntimeRDG/Source/NNERuntimeRDGUtils/Private/NNERuntimeRDGUtilsModelBuilderNNE.cpp b/Engine/Plugins/Experimental/NNERuntimeRDG/Source/NNERuntimeRDGUtils/Private/NNERuntimeRDGUtilsModelBuilderNNE.cpp
--- a/Engine/Plugins/Experimental/NNERuntimeRDG/Source/NNERuntimeRDGUtils/Private/NNERuntimeRDGUtilsModelBuilderNNE.cpp
+++ b/Engine/Plugins/Experimental/NNERuntimeRDG/Source/NNERuntimeRDGUtils/Private/NNERuntimeRDGUtilsModelBuilderNNE.cpp
@@ -4,18 +4,18 @@
 
 #include "NNEAttributeMap.h"
 #include "NNEHlslShadersLog.h"
-#include "NNERuntimeFormat.h"
+#include "NNERuntimeRDGDataFormat.h"
 #include "Misc/StringBuilder.h"
 #include "Serialization/MemoryWriter.h"
 
-namespace UE::NNERuntimeRDGUtils::Internal
+namespace UE::NNERuntimeRDGUtils::Private
 {
 
 namespace ModelBuilderNNEHelper
 {
-	int32 NNETensorCast(IModelBuilder::FHTensor& Handle)
+	int32 NNETensorCast(Internal::IModelBuilder::FHTensor& Handle)
 	{
-		if (Handle.Type == IModelBuilder::EHandleType::Tensor)
+		if (Handle.Type == Internal::IModelBuilder::EHandleType::Tensor)
 		{
 			return int32(reinterpret_cast<int64>(Handle.Ptr));
 		}
@@ -23,9 +23,9 @@ namespace ModelBuilderNNEHelper
 		return -1;
 	}
 
-	int32 NNEOperatorCast(IModelBuilder::FHOperator& Handle)
+	int32 NNEOperatorCast(Internal::IModelBuilder::FHOperator& Handle)
 	{
-		if (Handle.Type == IModelBuilder::EHandleType::Operator)
+		if (Handle.Type == Internal::IModelBuilder::EHandleType::Operator)
 		{
 			return int32(reinterpret_cast<int64>(Handle.Ptr));
 		}
@@ -34,7 +34,7 @@ namespace ModelBuilderNNEHelper
 	}
 } //namespace ModelBuilderNNEHelper
 
-class FModelBuilderNNE : public IModelBuilder
+class FModelBuilderNNE : public Internal::IModelBuilder
 {
 public:
 	virtual bool Begin(const FString& Name) override
@@ -68,7 +68,7 @@ public:
 
 	virtual FHTensor AddConstantTensor(const FString& Name, ENNETensorDataType DataType, TArrayView<const int32> Shape, const void* Data, uint64 DataSize) override
 	{
-		FNNEFormatTensorDesc Desc = GetBaseTensorDesc(Name, Shape, DataType);
+		FNNERuntimeRDGDataTensorDesc Desc = GetBaseTensorDesc(Name, Shape, DataType);
 		AddInitializer(Desc, Data, DataSize);
 		int32 Idx = AddTensor(MoveTemp(Desc));
 
@@ -92,13 +92,13 @@ public:
 			return false;
 		}
 
-		if (Format.Tensors[Idx].Type != ENNEFormatTensorType::None)
+		if (Format.Tensors[Idx].Type != ENNERuntimeRDGDataTensorType::None)
 		{
 			UE_LOG(LogNNERuntimeRDGHlsl, Warning, TEXT("Failed to add input tensor, tensor usage already set up"));
 			return false;
 		}
 
-		Format.Tensors[Idx].Type = ENNEFormatTensorType::Input;
+		Format.Tensors[Idx].Type = ENNERuntimeRDGDataTensorType::Input;
 
 		return true;
 	}
@@ -113,13 +113,13 @@ public:
 			return false;
 		}
 
-		if (Format.Tensors[Idx].Type != ENNEFormatTensorType::None)
+		if (Format.Tensors[Idx].Type != ENNERuntimeRDGDataTensorType::None)
 		{
 			UE_LOG(LogNNERuntimeRDGHlsl, Warning, TEXT("Failed to add output tensor, tensor usage already set up"));
 			return false;
 		}
 
-		Format.Tensors[Idx].Type = ENNEFormatTensorType::Output;
+		Format.Tensors[Idx].Type = ENNERuntimeRDGDataTensorType::Output;
 
 		return true;
 	}
@@ -128,7 +128,7 @@ public:
 	{
 		int32 Idx = Format.Operators.Num();
 
-		FNNEFormatOperatorDesc	Operator{};
+		FNNERuntimeRDGDataOperatorDesc	Operator{};
 
 		Operator.TypeName = TypeName;
 		Operator.DomainName = Domain;
@@ -165,15 +165,15 @@ public:
 			return false;
 		}
 
-		if (Format.Tensors[TensorIdx].Type == ENNEFormatTensorType::Input)
+		if (Format.Tensors[TensorIdx].Type == ENNERuntimeRDGDataTensorType::Input)
 		{
 			UE_LOG(LogNNERuntimeRDGHlsl, Warning, TEXT("Failed to add output tensor, tensor usage already set up to input"));
 			return false;
 		}
 
-		if (Format.Tensors[TensorIdx].Type == ENNEFormatTensorType::None)
+		if (Format.Tensors[TensorIdx].Type == ENNERuntimeRDGDataTensorType::None)
 		{
-			Format.Tensors[TensorIdx].Type = ENNEFormatTensorType::Intermediate;
+			Format.Tensors[TensorIdx].Type = ENNERuntimeRDGDataTensorType::Intermediate;
 		}
 
 		Format.Operators[OpIdx].OutTensors.Emplace(TensorIdx);
@@ -185,7 +185,7 @@ public:
 	{
 		int32 OpIdx = ModelBuilderNNEHelper::NNEOperatorCast(Op);
 
-		FNNEFormatAttributeDesc& Attribute = Format.Operators[OpIdx].Attributes.Emplace_GetRef();
+		FNNERuntimeRDGDataAttributeDesc& Attribute = Format.Operators[OpIdx].Attributes.Emplace_GetRef();
 		Attribute.Name = Name;
 		Attribute.Value = Value;
 
@@ -194,21 +194,21 @@ public:
 
 private:
 
-	static FNNEFormatTensorDesc GetBaseTensorDesc(const FString& InName, TArrayView<const int32> InShape, ENNETensorDataType InDataType)
+	static FNNERuntimeRDGDataTensorDesc GetBaseTensorDesc(const FString& InName, TArrayView<const int32> InShape, ENNETensorDataType InDataType)
 	{
-		FNNEFormatTensorDesc&&	Desc{};
+		FNNERuntimeRDGDataTensorDesc&&	Desc{};
 
 		Desc.Name = InName;
 		Desc.Shape = InShape;
-		Desc.Type = ENNEFormatTensorType::None;
+		Desc.Type = ENNERuntimeRDGDataTensorType::None;
 		Desc.DataType = InDataType;
 
 		return Desc;
 	}
 
-	void AddInitializer(FNNEFormatTensorDesc& Desc, const void* Data, uint64 DataSize = 0u)
+	void AddInitializer(FNNERuntimeRDGDataTensorDesc& Desc, const void* Data, uint64 DataSize = 0u)
 	{
-		Desc.Type = ENNEFormatTensorType::Initializer;
+		Desc.Type = ENNERuntimeRDGDataTensorType::Initializer;
 		Desc.DataSize = DataSize;
 
 		// Handle empty data initializers, i.e. when DataSize is 0
@@ -224,19 +224,19 @@ private:
 		return FString::Printf(TEXT("__NNE_EmptyTensor_%u"), EmptyTensorCounter++);
 	}
 
-	FNNEFormatTensorDesc GetEmptyTensorDesc()
+	FNNERuntimeRDGDataTensorDesc GetEmptyTensorDesc()
 	{
-		FNNEFormatTensorDesc&&	Desc{};
+		FNNERuntimeRDGDataTensorDesc&&	Desc{};
 
 		Desc.Name = GenerateEmptyTensorName();
 		Desc.Shape = { 0 };
-		Desc.Type = ENNEFormatTensorType::Empty;
+		Desc.Type = ENNERuntimeRDGDataTensorType::Empty;
 		Desc.DataType = ENNETensorDataType::None;
 
 		return Desc;
 	}
 
-	int32 AddTensor(FNNEFormatTensorDesc&& InTensorDesc)
+	int32 AddTensor(FNNERuntimeRDGDataTensorDesc&& InTensorDesc)
 	{
 		int32 Idx = -1;
 
@@ -258,14 +258,14 @@ private:
 	}
 
 
-	FNNERuntimeFormat	Format;
+	FNNERuntimeRDGDataModelFormat	Format;
 	TMap<FString, int32>	TensorMap;
 	uint32 EmptyTensorCounter = 0;
 };
 
-TUniquePtr<IModelBuilder> CreateNNEModelBuilder()
+TUniquePtr<Internal::IModelBuilder> CreateNNEModelBuilder()
 {
 	return MakeUnique<FModelBuilderNNE>();
 }
 
-} // namespace UE::NNERuntimeRDGUtils::Internal
+} // namespace UE::NNERuntimeRDGUtils::Private
diff --git a/Engine/Plugins/Experimental/NNERuntimeRDG/Source/NNERuntimeRDGUtils/Private/NNERuntimeRDGUtilsModelBuilderNNE.h b/Engine/Plugins/Experimental/NNERuntimeRDG/Source/NNERuntimeRDGUtils/Private/NNERuntimeRDGUtilsModelBuilderNNE.h
new file mode 100644
--- /dev/null
+++ b/Engine/Plugins/Experimental/NNERuntimeRDG/Source/NNERuntimeRDGUtils/Private/NNERuntimeRDGUtilsModelBuilderNNE.h
@@ -0,0 +1,14 @@
+// Copyright Epic Games, Inc. All Rights Reserved.
+
+#pragma once
+
+#include "NNERuntimeRDGUtilsModelBuilder.h"
+#include "Templates/UniquePtr.h"
+
+namespace UE::NNERuntimeRDGUtils::Private
+{
+
+TUniquePtr<NNERuntimeRDGUtils::Internal::IModelBuilder> CreateNNEModelBuilder();
+
+} // UE::NNERuntimeRDGUtils::Private
+
diff --git a/Engine/Plugins/Experimental/NNERuntimeRDG/Source/NNERuntimeRDGUtils/Private/NNERuntimeRDGUtilsModelBuilderONNX.cpp b/Engine/Plugins/Experimental/NNERuntimeRDG/Source/NNERuntimeRDGUtils/Private/NNERuntimeRDGUtilsModelBuilderONNX.cpp
--- a/Engine/Plugins/Experimental/NNERuntimeRDG/Source/NNERuntimeRDGUtils/Private/NNERuntimeRDGUtilsModelBuilderONNX.cpp
+++ b/Engine/Plugins/Experimental/NNERuntimeRDG/Source/NNERuntimeRDGUtils/Private/NNERuntimeRDGUtilsModelBuilderONNX.cpp
@@ -4,7 +4,6 @@
 
 #include "NNEAttributeMap.h"
 #include "NNEHlslShadersLog.h"
-#include "NNERuntimeFormat.h"
 
 THIRD_PARTY_INCLUDES_START
 #include "onnx/checker.h"
@@ -346,9 +345,9 @@ private:
 bool CreateONNXModelForOperator(const FString& OperatorName, int32 IrVersion, int32 OpsetVersion, bool bUseVariadicShapeForModel,
 	TConstArrayView<NNE::Internal::FTensor> InInputTensors, TConstArrayView<NNE::Internal::FTensor> InOutputTensors,
 	TConstArrayView<NNE::Internal::FTensor> InWeightTensors, TConstArrayView<TConstArrayView<uint8>> InWeightTensorsData,
-	const UE::NNE::FAttributeMap& Attributes, FNNEModelRaw& Model)
+	const UE::NNE::FAttributeMap& Attributes, TArray<uint8>& ModelData)
 {
-	Model = FNNEModelRaw{};
+	ModelData.Reset();
 
 	TUniquePtr<IModelBuilder> Builder(CreateONNXModelBuilder((int64) IrVersion, (int64) OpsetVersion));
 
@@ -422,8 +421,7 @@ bool CreateONNXModelForOperator(const FString& OperatorName, int32 IrVersion, in
 	}
 
 	
-	Builder->End(Model.Data);
-	Model.Format = ENNEInferenceFormat::ONNX;
+	Builder->End(ModelData);
 	
 	return true;
 }
diff --git a/Engine/Plugins/Experimental/NNERuntimeRDG/Source/NNERuntimeRDGUtils/Private/NNERuntimeRDGUtilsModelOptimizer.cpp b/Engine/Plugins/Experimental/NNERuntimeRDG/Source/NNERuntimeRDGUtils/Private/NNERuntimeRDGUtilsModelOptimizer.cpp
--- a/Engine/Plugins/Experimental/NNERuntimeRDG/Source/NNERuntimeRDGUtils/Private/NNERuntimeRDGUtilsModelOptimizer.cpp
+++ b/Engine/Plugins/Experimental/NNERuntimeRDG/Source/NNERuntimeRDGUtils/Private/NNERuntimeRDGUtilsModelOptimizer.cpp
@@ -3,30 +3,13 @@
 #include "NNERuntimeRDGUtilsModelOptimizer.h"
 
 #include "NNERuntimeRDGUtilsModelOptimizerNNE.h"
-#include "NNERuntimeRDGUtilsModelOptimizerONNX.h"
 
 namespace UE::NNERuntimeRDGUtils::Internal
 {
 
-TUniquePtr<IModelOptimizer> CreateModelOptimizer(ENNEInferenceFormat InputFormat, ENNEInferenceFormat OutputFormat)
+TUniquePtr<IModelOptimizer> CreateModelOptimizer()
 {
-	if (InputFormat == ENNEInferenceFormat::ONNX)
-	{
-		if (OutputFormat == ENNEInferenceFormat::NNERT)
-		{
-			return MakeUnique<FModelOptimizerONNXToNNERT>();
-		}
-		else if (OutputFormat == ENNEInferenceFormat::ONNX)
-		{
-			return MakeUnique<FModelOptimizerONNXToONNX>();
-		}
-		else
-		{
-			return MakeUnique<FModelOptimizerONNXToORT>();
-		}
-	}
-
-	return nullptr;
+	return MakeUnique<Private::FModelOptimizerONNXToNNERT>();
 }
 
 } // namespace UE::NNERuntimeRDGUtils::Internal
diff --git a/Engine/Plugins/Experimental/NNERuntimeRDG/Source/NNERuntimeRDGUtils/Private/NNERuntimeRDGUtilsModelOptimizerBase.cpp b/Engine/Plugins/Experimental/NNERuntimeRDG/Source/NNERuntimeRDGUtils/Private/NNERuntimeRDGUtilsModelOptimizerBase.cpp
--- a/Engine/Plugins/Experimental/NNERuntimeRDG/Source/NNERuntimeRDGUtils/Private/NNERuntimeRDGUtilsModelOptimizerBase.cpp
+++ b/Engine/Plugins/Experimental/NNERuntimeRDG/Source/NNERuntimeRDGUtils/Private/NNERuntimeRDGUtilsModelOptimizerBase.cpp
@@ -3,7 +3,7 @@
 #include "NNERuntimeRDGUtilsModelOptimizerBase.h"
 
 #include "NNEHlslShadersLog.h"
-#include "NNERuntimeFormat.h"
+#include "NNERuntimeRDGDataFormat.h"
 
 THIRD_PARTY_INCLUDES_START
 #include "onnx/common/common.h"
@@ -11,7 +11,7 @@ THIRD_PARTY_INCLUDES_START
 #include "onnx/proto_utils.h"
 THIRD_PARTY_INCLUDES_END
 
-namespace UE::NNERuntimeRDGUtils::Internal
+namespace UE::NNERuntimeRDGUtils::Private
 {
 
 FString FModelValidatorONNX::GetName() const
@@ -19,15 +19,10 @@ FString FModelValidatorONNX::GetName() const
 	return TEXT("ONNX Model validator");
 }
 
-bool FModelValidatorONNX::ValidateModel(const FNNEModelRaw& InputModel, const FOptimizerOptionsMap& Options) const
+bool FModelValidatorONNX::ValidateModel(TConstArrayView<uint8> InputModel) const
 {
-	if (InputModel.Format != ENNEInferenceFormat::ONNX)
-	{
-		return true;
-	}
-
 	onnx::ModelProto Model;
-	onnx::ParseProtoFromBytes(&Model, reinterpret_cast<const char*>(InputModel.Data.GetData()), static_cast<size_t>(InputModel.Data.Num()));
+	onnx::ParseProtoFromBytes(&Model, reinterpret_cast<const char*>(InputModel.GetData()), static_cast<size_t>(InputModel.Num()));
 
 #ifdef ONNX_NO_EXCEPTIONS
 	static_assert(false, "ONNX_NO_EXCEPTIONS is defined meaning onnx check_model would abort the program in case of validation failure.");
@@ -46,14 +41,14 @@ bool FModelValidatorONNX::ValidateModel(const FNNEModelRaw& InputModel, const FO
 	return true;
 }
 
-bool FModelOptimizerBase::IsModelValid(const FNNEModelRaw& ModelToValidate, const FOptimizerOptionsMap& Options)
+bool FModelOptimizerBase::IsModelValid(TConstArrayView<uint8> ModelToValidate)
 {
 	bool bIsModelValid = true;
 
-	for (TSharedPtr<IModelValidator>& Validator : Validators)
+	for (TSharedPtr<Internal::IModelValidator>& Validator : Validators)
 	{
 		check(Validator.IsValid());
-		if (!Validator->ValidateModel(ModelToValidate, Options))
+		if (!Validator->ValidateModel(ModelToValidate))
 		{
 			UE_LOG(LogNNERuntimeRDGHlsl, Warning, TEXT("Model validator '%s' detected an error."), *(Validator->GetName()));
 			bIsModelValid = false;
@@ -62,31 +57,31 @@ bool FModelOptimizerBase::IsModelValid(const FNNEModelRaw& ModelToValidate, cons
 	return bIsModelValid;
 }
 
-bool FModelOptimizerBase::ApplyAllPassesAndValidations(FNNEModelRaw& OptimizedModel, const FOptimizerOptionsMap& Options)
+bool FModelOptimizerBase::ApplyAllPassesAndValidations(TArray<uint8>& OptimizedModel)
 {
-	if (!IsModelValid(OptimizedModel, Options))
+	if (!IsModelValid(OptimizedModel))
 	{
 		UE_LOG(LogNNERuntimeRDGHlsl, Warning, TEXT("Model is not valid."));
 		return false;
 	}
 		
-	for (TSharedPtr<IModelOptimizerPass>& Pass : OptimizationPasses)
+	for (TSharedPtr<Internal::IModelOptimizerPass>& Pass : OptimizationPasses)
 	{
 		check(Pass.IsValid());
 
-		//Note: Usefull to enable for debug purpose
+		//Note: Useful to enable for debug purpose
 		//FFileHelper::SaveArrayToFile(OptimizedModel.Data, TEXT("D:\\OnnxBeforePass.onnx"));
 			
-		if (!Pass->ApplyPass(OptimizedModel, Options))
+		if (!Pass->ApplyPass(OptimizedModel))
 		{
 			UE_LOG(LogNNERuntimeRDGHlsl, Warning, TEXT("Error while executing model optimisation pass '%s'."), *(Pass->GetName()));
 			return false;
 		}
 
-		//Note: Usefull to enable for debug purpose
+		//Note: Useful to enable for debug purpose
 		//FFileHelper::SaveArrayToFile(OptimizedModel.Data, TEXT("D:\\OnnxAfterPass.onnx"));
 
-		if (!IsModelValid(OptimizedModel, Options))
+		if (!IsModelValid(OptimizedModel))
 		{
 			UE_LOG(LogNNERuntimeRDGHlsl, Warning, TEXT("Model validation failed after optimisation pass '%s'."), *(Pass->GetName()));
 			return false;
@@ -96,7 +91,7 @@ bool FModelOptimizerBase::ApplyAllPassesAndValidations(FNNEModelRaw& OptimizedMo
 	return true;
 }
 
-void FModelOptimizerBase::AddOptimizationPass(TSharedPtr<IModelOptimizerPass> ModelOptimizerPass)
+void FModelOptimizerBase::AddOptimizationPass(TSharedPtr<Internal::IModelOptimizerPass> ModelOptimizerPass)
 {
 	if (ModelOptimizerPass.IsValid())
 	{
@@ -104,7 +99,7 @@ void FModelOptimizerBase::AddOptimizationPass(TSharedPtr<IModelOptimizerPass> Mo
 	}
 }
 
-void FModelOptimizerBase::AddValidator(TSharedPtr<IModelValidator> ModelValidator)
+void FModelOptimizerBase::AddValidator(TSharedPtr<Internal::IModelValidator> ModelValidator)
 {
 	if (ModelValidator.IsValid())
 	{
@@ -112,10 +107,10 @@ void FModelOptimizerBase::AddValidator(TSharedPtr<IModelValidator> ModelValidato
 	}
 }
 
-bool FModelOptimizerBase::Optimize(const FNNEModelRaw& InputModel, FNNEModelRaw& OptimizedModel, const FOptimizerOptionsMap& Options)
+bool FModelOptimizerBase::Optimize(TConstArrayView<uint8> InputModel, TArray<uint8>& OutModel)
 {
-	OptimizedModel = InputModel;
-	return ApplyAllPassesAndValidations(OptimizedModel, Options);
+	OutModel = InputModel;
+	return ApplyAllPassesAndValidations(OutModel);
 }
 
-} // namespace UE::NNERuntimeRDGUtils::Internal
+} // namespace UE::NNERuntimeRDGUtils::Private
diff --git a/Engine/Plugins/Experimental/NNERuntimeRDG/Source/NNERuntimeRDGUtils/Private/NNERuntimeRDGUtilsModelOptimizerBase.h b/Engine/Plugins/Experimental/NNERuntimeRDG/Source/NNERuntimeRDGUtils/Private/NNERuntimeRDGUtilsModelOptimizerBase.h
new file mode 100644
--- /dev/null
+++ b/Engine/Plugins/Experimental/NNERuntimeRDG/Source/NNERuntimeRDGUtils/Private/NNERuntimeRDGUtilsModelOptimizerBase.h
@@ -0,0 +1,33 @@
+// Copyright Epic Games, Inc. All Rights Reserved.
+#pragma once
+
+#include "CoreMinimal.h"
+#include "NNERuntimeRDGUtilsModelOptimizerInterface.h"
+#include "Templates/SharedPointer.h"
+
+namespace UE::NNERuntimeRDGUtils::Private
+{
+class FModelValidatorONNX : public Internal::IModelValidator
+{
+public:
+	virtual FString GetName() const override;
+	virtual bool ValidateModel(TConstArrayView<uint8> InputModel) const override;
+};
+
+class FModelOptimizerBase : public Internal::IModelOptimizer
+{
+protected:
+	FModelOptimizerBase() {}
+	TArray<TSharedPtr<Internal::IModelOptimizerPass>> OptimizationPasses;
+	TArray<TSharedPtr<Internal::IModelValidator>> Validators;
+
+	bool IsModelValid(TConstArrayView<uint8> ModelToValidate);
+	bool ApplyAllPassesAndValidations(TArray<uint8>& OptimizedModel);
+
+public:
+	virtual void AddOptimizationPass(TSharedPtr<Internal::IModelOptimizerPass> ModelOptimizerPass) override;
+	virtual void AddValidator(TSharedPtr<Internal::IModelValidator> ModelValidator) override;
+	virtual bool Optimize(TConstArrayView<uint8> InputModel, TArray<uint8>& OutModel) override;
+};
+
+} // namespace UE::NNERuntimeRDGUtils::Private
diff --git a/Engine/Plugins/Experimental/NNERuntimeRDG/Source/NNERuntimeRDGUtils/Private/NNERuntimeRDGUtilsModelOptimizerNNE.cpp b/Engine/Plugins/Experimental/NNERuntimeRDG/Source/NNERuntimeRDGUtils/Private/NNERuntimeRDGUtilsModelOptimizerNNE.cpp
--- a/Engine/Plugins/Experimental/NNERuntimeRDG/Source/NNERuntimeRDGUtils/Private/NNERuntimeRDGUtilsModelOptimizerNNE.cpp
+++ b/Engine/Plugins/Experimental/NNERuntimeRDG/Source/NNERuntimeRDGUtils/Private/NNERuntimeRDGUtilsModelOptimizerNNE.cpp
@@ -5,7 +5,7 @@
 #include "Containers/ContainersFwd.h"
 #include "NNEAttributeTensor.h"
 #include "NNEHlslShadersLog.h"
-#include "NNERuntimeFormat.h"
+#include "NNERuntimeRDGDataFormat.h"
 #include "NNERuntimeRDGUtilsModelBuilderNNE.h"
 #include "NNERuntimeRDGUtilsModelOptimizerONNX.h"
 #include "NNERuntimeRDGUtilsHelpers.h"
@@ -18,12 +18,13 @@ THIRD_PARTY_INCLUDES_START
 #include <onnx/defs/schema.h>
 THIRD_PARTY_INCLUDES_END
 
-namespace UE::NNERuntimeRDGUtils::Internal
+namespace UE::NNERuntimeRDGUtils::Private
 {
 
 namespace ModelOptimizerNNEHelper
 {
-	
+	using namespace UE::NNERuntimeRDGUtils::Internal;
+
 	ENNETensorDataType GetNNETensorTypeFromONNX(int32 DataType)
 	{
 		switch (DataType)
@@ -560,27 +561,26 @@ namespace ModelOptimizerNNEHelper
 
 } // namespace ModelOptimizerNNEHelper
 
-bool FModelOptimizerONNXToNNERT::Optimize(const FNNEModelRaw& InputModel, FNNEModelRaw& OptimizedModel, const FOptimizerOptionsMap& Options)
+bool FModelOptimizerONNXToNNERT::Optimize(TConstArrayView<uint8> InputModel, TArray<uint8>& OptimizedModel)
 {
-	OptimizedModel = FNNEModelRaw{};
+	OptimizedModel.Reset();
 
-	FNNEModelRaw OptimizedONNXModel = FNNEModelRaw{};
+	TArray<uint8> OptimizedONNXModel;
 	FModelOptimizerONNXToONNX ModelOptimizerONNXToONNX;
 
-	if (!ModelOptimizerONNXToONNX.Optimize(InputModel, OptimizedONNXModel, Options))
+	if (!ModelOptimizerONNXToONNX.Optimize(InputModel, OptimizedONNXModel))
 	{
 		UE_LOG(LogNNERuntimeRDGHlsl, Warning, TEXT("Error while optimizing the ONNX model before convertion to NNERT format."));
 		return false;
 	}
 
-	if (!ModelOptimizerNNEHelper::BuildNNEFormatFromONNX(OptimizedONNXModel.Data, OptimizedModel.Data))
+	if (!ModelOptimizerNNEHelper::BuildNNEFormatFromONNX(OptimizedONNXModel, OptimizedModel))
 	{
 		UE_LOG(LogNNERuntimeRDGHlsl, Warning, TEXT("Error while building NNERT Model from ONNX."));
 		return false;
 	}
-	OptimizedModel.Format = ENNEInferenceFormat::NNERT;
 
-	return ApplyAllPassesAndValidations(OptimizedModel, Options);
+	return ApplyAllPassesAndValidations(OptimizedModel);
 }
 
-} // namespace UE::NNERuntimeRDGUtils::Internal
+} // namespace UE::NNERuntimeRDGUtils::Private
diff --git a/Engine/Plugins/Experimental/NNERuntimeRDG/Source/NNERuntimeRDGUtils/Private/NNERuntimeRDGUtilsModelOptimizerNNE.h b/Engine/Plugins/Experimental/NNERuntimeRDG/Source/NNERuntimeRDGUtils/Private/NNERuntimeRDGUtilsModelOptimizerNNE.h
new file mode 100644
--- /dev/null
+++ b/Engine/Plugins/Experimental/NNERuntimeRDG/Source/NNERuntimeRDGUtils/Private/NNERuntimeRDGUtilsModelOptimizerNNE.h
@@ -0,0 +1,21 @@
+// Copyright Epic Games, Inc. All Rights Reserved.
+#pragma once
+
+#include "CoreMinimal.h"
+#include "NNERuntimeRDGUtilsModelOptimizerBase.h"
+
+namespace UE::NNERuntimeRDGUtils::Private
+{
+
+class FModelOptimizerONNXToNNERT : public FModelOptimizerBase
+{
+public:
+	virtual FString GetName() const override 
+	{
+		return TEXT("NNEModelOptimizerFromONNXToNNERuntimeRDGHlsl");
+	}
+
+	virtual bool Optimize(TConstArrayView<uint8> InputModel, TArray<uint8>& OutModel) override;
+};
+
+} // namespace UE::NNERuntimeRDGUtils::Private
diff --git a/Engine/Plugins/Experimental/NNERuntimeRDG/Source/NNERuntimeRDGUtils/Private/NNERuntimeRDGUtilsModelOptimizerONNX.cpp b/Engine/Plugins/Experimental/NNERuntimeRDG/Source/NNERuntimeRDGUtils/Private/NNERuntimeRDGUtilsModelOptimizerONNX.cpp
--- a/Engine/Plugins/Experimental/NNERuntimeRDG/Source/NNERuntimeRDGUtils/Private/NNERuntimeRDGUtilsModelOptimizerONNX.cpp
+++ b/Engine/Plugins/Experimental/NNERuntimeRDG/Source/NNERuntimeRDGUtils/Private/NNERuntimeRDGUtilsModelOptimizerONNX.cpp
@@ -15,17 +15,14 @@ THIRD_PARTY_INCLUDES_START
 #include <onnx/shape_inference/implementation.h>
 THIRD_PARTY_INCLUDES_END
 
-namespace UE::NNERuntimeRDGUtils::Internal
+namespace UE::NNERuntimeRDGUtils::Private
 {
 
-class FOnnxRuntimeModelOptimizerPass : public IModelOptimizerPass
+class FOnnxRuntimeModelOptimizerPass : public Internal::IModelOptimizerPass
 {
-	ENNEInferenceFormat	TargetFormat;
-
 public:
-	FOnnxRuntimeModelOptimizerPass(ENNEInferenceFormat OutFormat) : TargetFormat(OutFormat)
+	FOnnxRuntimeModelOptimizerPass()
 	{
-		check(TargetFormat == ENNEInferenceFormat::ONNX || TargetFormat == ENNEInferenceFormat::ORT);
 	}
 
 	virtual FString GetName() const
@@ -33,16 +30,10 @@ public:
 		return TEXT("Onnx runtime model optimization");
 	}
 
-	virtual bool ApplyPass(FNNEModelRaw& Model, const FOptimizerOptionsMap& Options) const
+	virtual bool ApplyPass(TArray<uint8>& ModelData) const
 	{
-		if (Model.Format != ENNEInferenceFormat::ONNX)
-		{
-			UE_LOG(LogNNERuntimeRDGHlsl, Warning, TEXT("%s is expecting a model in ONNX format but received %u."), *(GetName()), Model.Format);
-			return false;
-		}
-
 		onnx::ModelProto ModelProto;
-		const bool result = ModelProto.ParseFromArray(Model.Data.GetData(), Model.Data.Num());
+		const bool result = ModelProto.ParseFromArray(ModelData.GetData(), ModelData.Num());
 		if (!result)
 		{
 			UE_LOG(LogNNERuntimeRDGHlsl, Warning, TEXT("%s could not parse the input model as a ModelProto."), *(GetName()));
@@ -65,14 +56,12 @@ public:
 
 		FString ProjIntermediateDir = FPaths::ConvertRelativePathToFull(FPaths::ProjectIntermediateDir());
 		FString ModelToOptimizePath = FPaths::CreateTempFilename(*ProjIntermediateDir, TEXT("ORTOptimizerPass_ToOptimize"), TEXT(".onnx"));
-		FString TargetExtension = TargetFormat == ENNEInferenceFormat::ONNX ? TEXT(".onnx") : TEXT(".ort");
-		FString ModelOptimizedPath = FPaths::CreateTempFilename(*ProjIntermediateDir, TEXT("ORTOptimizerPass_Optimized"), *TargetExtension);
+		FString ModelOptimizedPath = FPaths::CreateTempFilename(*ProjIntermediateDir, TEXT("ORTOptimizerPass_Optimized"), TEXT(".onnx"));
 
 		//See https://onnxruntime.ai/docs/performance/model-optimizations/graph-optimizations.html
 		//We only enable all the optimization when going to ORT format itself for the CPU provider
-		GraphOptimizationLevel OptimizationLevel = TargetFormat == ENNEInferenceFormat::ONNX ? ORT_ENABLE_BASIC : ORT_ENABLE_ALL;
-
-		FFileHelper::SaveArrayToFile(Model.Data, *ModelToOptimizePath);
+		GraphOptimizationLevel OptimizationLevel = ORT_ENABLE_BASIC;
+		FFileHelper::SaveArrayToFile(ModelData, *ModelToOptimizePath);
 		{
 			Ort::ThreadingOptions ThreadingOptions;
 			ThreadingOptions.SetGlobalIntraOpNumThreads(1);
@@ -93,13 +82,11 @@ public:
 				Ort::Session Session(Env, TCHAR_TO_ANSI(*ModelToOptimizePath), SessionOptions);
 			#endif
 		}
-		FFileHelper::LoadFileToArray(Model.Data, *ModelOptimizedPath);
+		FFileHelper::LoadFileToArray(ModelData, *ModelOptimizedPath);
 
 		IFileManager::Get().Delete(*ModelToOptimizePath);
 		IFileManager::Get().Delete(*ModelOptimizedPath);
 
-		Model.Format = TargetFormat;
-
 		return true;
 	}
 private:
@@ -144,7 +131,7 @@ private:
 
 };
 
-class FOnnxDomainCleanupModelOptimizerPass : public IModelOptimizerPass
+class FOnnxDomainCleanupModelOptimizerPass : public Internal::IModelOptimizerPass
 {
 public:
 	virtual FString GetName() const
@@ -152,16 +139,10 @@ public:
 		return TEXT("Onnx domain cleanup");
 	}
 
-	virtual bool ApplyPass(FNNEModelRaw& Model, const FOptimizerOptionsMap& Options) const
+	virtual bool ApplyPass(TArray<uint8>& ModelData) const
 	{
-		if (Model.Format != ENNEInferenceFormat::ONNX)
-		{
-			UE_LOG(LogNNERuntimeRDGHlsl, Warning, TEXT("%s is expecting a model in ONNX format but received %u."), *(GetName()), Model.Format);
-			return false;
-		}
-
 		onnx::ModelProto ModelProto;
-		const bool result = ModelProto.ParseFromArray(Model.Data.GetData(), Model.Data.Num());
+		const bool result = ModelProto.ParseFromArray(ModelData.GetData(), ModelData.Num());
 		if (!result)
 		{
 			UE_LOG(LogNNERuntimeRDGHlsl, Warning, TEXT("%s could not parse the input model as a ModelProto."), *(GetName()));
@@ -194,15 +175,15 @@ public:
 		ModelProto.mutable_opset_import()->Clear();
 		ModelProto.mutable_opset_import()->Add(UsedOperatorSet.cbegin(), UsedOperatorSet.cend());
 
-		Model.Data.SetNumUninitialized(ModelProto.ByteSizeLong());
-		ModelProto.SerializeToArray(Model.Data.GetData(), Model.Data.Num());
+		ModelData.SetNumUninitialized(ModelProto.ByteSizeLong());
+		ModelProto.SerializeToArray(ModelData.GetData(), ModelData.Num());
 
 		return true;
 	}
 
 };
 
-class FOnnxShapeInferenceModelOptimizerPass : public IModelOptimizerPass
+class FOnnxShapeInferenceModelOptimizerPass : public Internal::IModelOptimizerPass
 {
 public:
 	virtual FString GetName() const
@@ -210,16 +191,10 @@ public:
 		return TEXT("Onnx shape inference");
 	}
 
-	virtual bool ApplyPass(FNNEModelRaw& Model, const FOptimizerOptionsMap& Options) const
+	virtual bool ApplyPass(TArray<uint8>& ModelData) const
 	{
-		if (Model.Format != ENNEInferenceFormat::ONNX)
-		{
-			UE_LOG(LogNNERuntimeRDGHlsl, Warning, TEXT("%s is expecting a model in ONNX format but received %u."), *(GetName()), Model.Format);
-			return false;
-		}
-
 		onnx::ModelProto ModelProto;
-		const bool result = ModelProto.ParseFromArray(Model.Data.GetData(), Model.Data.Num());
+		const bool result = ModelProto.ParseFromArray(ModelData.GetData(), ModelData.Num());
 		if (!result)
 		{
 			UE_LOG(LogNNERuntimeRDGHlsl, Warning, TEXT("%s could not parse the input model as a ModelProto."), *(GetName()));
@@ -243,7 +218,7 @@ public:
 		}
 #endif
 		
-		ModelProto.SerializeToArray(Model.Data.GetData(), Model.Data.Num());
+		ModelProto.SerializeToArray(ModelData.GetData(), ModelData.Num());
 
 		return true;
 	}
@@ -252,16 +227,10 @@ public:
 
 FModelOptimizerONNXToONNX::FModelOptimizerONNXToONNX()
 {
-	AddOptimizationPass(MakeShared<FOnnxRuntimeModelOptimizerPass>(ENNEInferenceFormat::ONNX));
+	AddOptimizationPass(MakeShared<FOnnxRuntimeModelOptimizerPass>());
 	AddOptimizationPass(MakeShared<FOnnxDomainCleanupModelOptimizerPass>());
 	AddOptimizationPass(MakeShared<FOnnxShapeInferenceModelOptimizerPass>());
 	AddValidator(MakeShared<FModelValidatorONNX>());
 }
 
-FModelOptimizerONNXToORT::FModelOptimizerONNXToORT()
-{
-	AddOptimizationPass(MakeShared<FOnnxRuntimeModelOptimizerPass>(ENNEInferenceFormat::ORT));
-	AddValidator(MakeShared<FModelValidatorONNX>());
-}
-
-} // namespace UE::NNERuntimeRDGUtils::Internal
+} // namespace UE::NNERuntimeRDGUtils::Private
diff --git a/Engine/Plugins/Experimental/NNERuntimeRDG/Source/NNERuntimeRDGUtils/Private/NNERuntimeRDGUtilsModelOptimizerONNX.h b/Engine/Plugins/Experimental/NNERuntimeRDG/Source/NNERuntimeRDGUtils/Private/NNERuntimeRDGUtilsModelOptimizerONNX.h
new file mode 100644
--- /dev/null
+++ b/Engine/Plugins/Experimental/NNERuntimeRDG/Source/NNERuntimeRDGUtils/Private/NNERuntimeRDGUtilsModelOptimizerONNX.h
@@ -0,0 +1,21 @@
+// Copyright Epic Games, Inc. All Rights Reserved.
+#pragma once
+
+#include "CoreMinimal.h"
+#include "NNERuntimeRDGUtilsModelOptimizerBase.h"
+
+namespace UE::NNERuntimeRDGUtils::Private
+{
+
+class FModelOptimizerONNXToONNX : public FModelOptimizerBase
+{
+public:
+	FModelOptimizerONNXToONNX();
+
+	virtual FString GetName() const override
+	{
+		return TEXT("NNEModelOptimizerFromONNXToONNX");
+	}
+};
+
+} // namespace UE::NNERuntimeRDGUtils::Private
diff --git a/Engine/Source/Runtime/NNE/Internal/NNERuntimeFormat.h b/Engine/Source/Runtime/NNE/Internal/NNERuntimeFormat.h
deleted file mode 100644
diff --git a/Engine/Source/Runtime/NNE/Public/NNETypes.h b/Engine/Source/Runtime/NNE/Public/NNETypes.h
--- a/Engine/Source/Runtime/NNE/Public/NNETypes.h
+++ b/Engine/Source/Runtime/NNE/Public/NNETypes.h
@@ -14,7 +14,7 @@
 UENUM()
 enum class ENNETensorDataType : uint8
 {
-	None,								//!< For empty tensors (ENNEFormatTensorType::Empty)
+	None,								//!< For empty tensors
 	Char,								//!< Character type
 	Boolean,							//!< Boolean type
 	Half,								//!< 16-bit floating number
