[NNE] NNE internal cleanup step4 : operator attributes

Problem: NNE plugin are dependant on Source/NNE/Internal
Solution (step4):
* remove NNEAttributes*.h and move to NNERuntimeRDGUtils and NNEQA
* break the dependancy between NNERuntimeRDGUtils and NNEQA by adding a copy of the onnxeditor module to NNEQA plugins.
* remove model generation capability from NNERuntimeRDGUtils (instead they are directly in NNEQAEditor)

Test:
Win64 editor debug : NNEQA all tests (+reimported) + style transfer
Win64 package debug : NNEQA all tests (+reimported) + style transfer
Mac editor dev : NNEQA smoketests (+reimported) + style transfer
Console package dev : NNEQA smoketests (+reimported) + style transfer

#rnx

diff --git a/Engine/Plugins/Experimental/NNERuntimeCoreML/Source/NNERuntimeCoreML/Private/NNERuntimeCoreMLModel.cpp b/Engine/Plugins/Experimental/NNERuntimeCoreML/Source/NNERuntimeCoreML/Private/NNERuntimeCoreMLModel.cpp
--- a/Engine/Plugins/Experimental/NNERuntimeCoreML/Source/NNERuntimeCoreML/Private/NNERuntimeCoreMLModel.cpp
+++ b/Engine/Plugins/Experimental/NNERuntimeCoreML/Source/NNERuntimeCoreML/Private/NNERuntimeCoreMLModel.cpp
@@ -11,8 +11,6 @@
 #include "Misc/Paths.h"
 #include "NNEModelData.h"
 #include "NNERuntimeCoreML.h"
-#include "NNETensor.h"
-#include "NNETypes.h"
 
 #if defined(__APPLE__)
 #import <Foundation/Foundation.h>
@@ -274,8 +272,8 @@ FModelInstanceCoreMLCpu::ESetInputTensorShapesStatus FModelInstanceCoreMLBase<In
 	{
 		if (SymbolicTensorDesc.GetShape().IsConcrete())
 		{
-			NNE::Internal::FTensor Tensor = NNE::Internal::FTensor::MakeFromSymbolicDesc(SymbolicTensorDesc);
-			OutputTensorShapes.Emplace(Tensor.GetShape());
+			NNE::FTensorShape TensorShape = NNE::FTensorShape::MakeFromSymbolic(SymbolicTensorDesc.GetShape());
+			OutputTensorShapes.Emplace(TensorShape);
 		}
 	}
 	if (OutputTensorShapes.Num() != OutputSymbolicTensors.Num())
diff --git a/Engine/Plugins/Experimental/NNERuntimeCoreML/Source/NNERuntimeCoreML/Private/NNERuntimeCoreMLModel.h b/Engine/Plugins/Experimental/NNERuntimeCoreML/Source/NNERuntimeCoreML/Private/NNERuntimeCoreMLModel.h
--- a/Engine/Plugins/Experimental/NNERuntimeCoreML/Source/NNERuntimeCoreML/Private/NNERuntimeCoreMLModel.h
+++ b/Engine/Plugins/Experimental/NNERuntimeCoreML/Source/NNERuntimeCoreML/Private/NNERuntimeCoreMLModel.h
@@ -8,7 +8,6 @@
 #include "NNERuntimeCPU.h"
 #include "NNERuntimeGPU.h"
 #include "NNERuntimeNPU.h"
-#include "NNETensor.h"
 #include "NNETypes.h"
 
 #if defined(__APPLE__)
diff --git a/Engine/Plugins/Experimental/NNERuntimeRDG/Source/NNERuntimeRDG/Private/Hlsl/NNERuntimeRDGBatchNormalization.cpp b/Engine/Plugins/Experimental/NNERuntimeRDG/Source/NNERuntimeRDG/Private/Hlsl/NNERuntimeRDGBatchNormalization.cpp
--- a/Engine/Plugins/Experimental/NNERuntimeRDG/Source/NNERuntimeRDG/Private/Hlsl/NNERuntimeRDGBatchNormalization.cpp
+++ b/Engine/Plugins/Experimental/NNERuntimeRDG/Source/NNERuntimeRDG/Private/Hlsl/NNERuntimeRDGBatchNormalization.cpp
@@ -4,8 +4,8 @@
 
 #include "NNEHlslShadersBatchNormalizationCS.h"
 #include "NNEHlslShadersLog.h"
+#include "NNERuntimeRDGDataAttributeMap.h"
 #include "NNERuntimeRDGHlslHelper.h"
-#include "NNEAttributeMap.h"
 #include "NNETypes.h"
 #include "NNETensor.h"
 
@@ -39,7 +39,7 @@ namespace UE::NNERuntimeRDG::Private::Hlsl
 			return 0;
 		};
 
-		virtual bool Initialize(TConstArrayView<NNE::FTensorDesc> InputTensorDescs, TConstArrayView<NNE::FTensorDesc> OutputTensorDescs, const NNE::FAttributeMap& Attributes) override
+		virtual bool Initialize(TConstArrayView<NNE::FTensorDesc> InputTensorDescs, TConstArrayView<NNE::FTensorDesc> OutputTensorDescs, const NNERuntimeRDGData::Internal::FAttributeMap& Attributes) override
 		{
 			check(InputTensorDescs.Num() == 5);
 			check(OutputTensorDescs.Num() >= 1);
@@ -124,13 +124,13 @@ namespace UE::NNERuntimeRDG::Private::Hlsl
 		}
 	};
 
-	bool ValidateBatchNormalizationOperator(const NNE::FAttributeMap& AttributeMap, TConstArrayView<ENNETensorDataType> InputTypes, TConstArrayView<NNE::FSymbolicTensorShape> InputBatchNormalizations)
+	bool ValidateBatchNormalizationOperator(const NNERuntimeRDGData::Internal::FAttributeMap& AttributeMap, TConstArrayView<ENNETensorDataType> InputTypes, TConstArrayView<NNE::FSymbolicTensorShape> InputBatchNormalizations)
 	{
 		bool bIsValid = true;
 
 		FAttributeValidator AttributeValidator;
-		AttributeValidator.AddOptional(TEXT("epsilon"), ENNEAttributeDataType::Float);
-		AttributeValidator.AddOptional(TEXT("momentum"), ENNEAttributeDataType::Float); //Will be ignored, only useful in training mode.
+		AttributeValidator.AddOptional(TEXT("epsilon"), ENNERuntimeRDGDataAttributeDataType::Float);
+		AttributeValidator.AddOptional(TEXT("momentum"), ENNERuntimeRDGDataAttributeDataType::Float); //Will be ignored, only useful in training mode.
 		bIsValid &= AttributeValidator.Validate(AttributeMap);
 
 		FInputValidator InputValidator;
diff --git a/Engine/Plugins/Experimental/NNERuntimeRDG/Source/NNERuntimeRDG/Private/Hlsl/NNERuntimeRDGCast.cpp b/Engine/Plugins/Experimental/NNERuntimeRDG/Source/NNERuntimeRDG/Private/Hlsl/NNERuntimeRDGCast.cpp
--- a/Engine/Plugins/Experimental/NNERuntimeRDG/Source/NNERuntimeRDG/Private/Hlsl/NNERuntimeRDGCast.cpp
+++ b/Engine/Plugins/Experimental/NNERuntimeRDG/Source/NNERuntimeRDG/Private/Hlsl/NNERuntimeRDGCast.cpp
@@ -47,7 +47,7 @@ namespace UE::NNERuntimeRDG::Private::Hlsl
 			return 0;
 		}
 
-		virtual bool Initialize(TConstArrayView<NNE::FTensorDesc> InputTensorDescs, TConstArrayView<NNE::FTensorDesc> OutputTensorDescs, const NNE::FAttributeMap& Attributes) override
+		virtual bool Initialize(TConstArrayView<NNE::FTensorDesc> InputTensorDescs, TConstArrayView<NNE::FTensorDesc> OutputTensorDescs, const NNERuntimeRDGData::Internal::FAttributeMap& Attributes) override
 		{
 			check(InputTensorDescs.Num() == 1);
 			check(OutputTensorDescs.Num() == 1);
@@ -107,12 +107,12 @@ namespace UE::NNERuntimeRDG::Private::Hlsl
 		}
 	};
 
-	bool ValidateCastOperator(const NNE::FAttributeMap& AttributeMap, TConstArrayView<ENNETensorDataType> InputTypes, TConstArrayView<NNE::FSymbolicTensorShape> InputShapes)
+	bool ValidateCastOperator(const NNERuntimeRDGData::Internal::FAttributeMap& AttributeMap, TConstArrayView<ENNETensorDataType> InputTypes, TConstArrayView<NNE::FSymbolicTensorShape> InputShapes)
 	{
 		bool bIsValid = true;
 
 		FAttributeValidator AttributeValidator;
-		AttributeValidator.AddRequired(TEXT("to"), ENNEAttributeDataType::Int32);
+		AttributeValidator.AddRequired(TEXT("to"), ENNERuntimeRDGDataAttributeDataType::Int32);
 		bIsValid &= AttributeValidator.Validate(AttributeMap);
 
 		if (bIsValid)
diff --git a/Engine/Plugins/Experimental/NNERuntimeRDG/Source/NNERuntimeRDG/Private/Hlsl/NNERuntimeRDGConcat.cpp b/Engine/Plugins/Experimental/NNERuntimeRDG/Source/NNERuntimeRDG/Private/Hlsl/NNERuntimeRDGConcat.cpp
--- a/Engine/Plugins/Experimental/NNERuntimeRDG/Source/NNERuntimeRDG/Private/Hlsl/NNERuntimeRDGConcat.cpp
+++ b/Engine/Plugins/Experimental/NNERuntimeRDG/Source/NNERuntimeRDG/Private/Hlsl/NNERuntimeRDGConcat.cpp
@@ -83,7 +83,7 @@ namespace UE::NNERuntimeRDG::Private::Hlsl
 			return 0;
 		};
 
-		virtual bool Initialize(TConstArrayView<NNE::FTensorDesc> InputTensorDescs, TConstArrayView<NNE::FTensorDesc> OutputTensorDescs, const NNE::FAttributeMap& Attributes) override
+		virtual bool Initialize(TConstArrayView<NNE::FTensorDesc> InputTensorDescs, TConstArrayView<NNE::FTensorDesc> OutputTensorDescs, const NNERuntimeRDGData::Internal::FAttributeMap& Attributes) override
 		{
 			check(InputTensorDescs.Num() >= 1);
 			check(OutputTensorDescs.Num() == 1);
@@ -146,12 +146,12 @@ namespace UE::NNERuntimeRDG::Private::Hlsl
 		}
 	};
 
-	bool ValidateConcatOperator(const NNE::FAttributeMap& AttributeMap, TConstArrayView<ENNETensorDataType> InputTypes, TConstArrayView<NNE::FSymbolicTensorShape> InputConcats)
+	bool ValidateConcatOperator(const NNERuntimeRDGData::Internal::FAttributeMap& AttributeMap, TConstArrayView<ENNETensorDataType> InputTypes, TConstArrayView<NNE::FSymbolicTensorShape> InputConcats)
 	{
 		bool bIsValid = true;
 
 		FAttributeValidator AttributeValidator;
-		AttributeValidator.AddRequired(TEXT("axis"), ENNEAttributeDataType::Int32);
+		AttributeValidator.AddRequired(TEXT("axis"), ENNERuntimeRDGDataAttributeDataType::Int32);
 		bIsValid &= AttributeValidator.Validate(AttributeMap);
 
 		if (InputTypes.Num() == 0)
diff --git a/Engine/Plugins/Experimental/NNERuntimeRDG/Source/NNERuntimeRDG/Private/Hlsl/NNERuntimeRDGConstant.cpp b/Engine/Plugins/Experimental/NNERuntimeRDG/Source/NNERuntimeRDG/Private/Hlsl/NNERuntimeRDGConstant.cpp
--- a/Engine/Plugins/Experimental/NNERuntimeRDG/Source/NNERuntimeRDG/Private/Hlsl/NNERuntimeRDGConstant.cpp
+++ b/Engine/Plugins/Experimental/NNERuntimeRDG/Source/NNERuntimeRDG/Private/Hlsl/NNERuntimeRDGConstant.cpp
@@ -4,8 +4,8 @@
 
 #include "Helper/NNERuntimeRDGLogHelper.h"
 #include "Helper/NNERuntimeRDGOperatorHelper.h"
-#include "NNEAttributeTensor.h"
 #include "NNEHlslShadersLog.h"
+#include "NNERuntimeRDGDataAttributeTensor.h"
 #include "NNERuntimeRDGHlslHelper.h"
 #include "NNETensor.h"
 #include "NNETypes.h"
@@ -24,7 +24,7 @@ namespace UE::NNERuntimeRDG::Private::Hlsl
 
 	private:
 
-		FNNEAttributeValue Attribute;
+		FNNERuntimeRDGDataAttributeValue Attribute;
 
 
 	public:
@@ -33,7 +33,7 @@ namespace UE::NNERuntimeRDG::Private::Hlsl
 		{
 			using NNE::FTensorShape;
 			using NNE::Internal::FTensor;
-			using NNE::Internal::FAttributeTensor;
+			using NNERuntimeRDGData::Internal::FAttributeTensor;
 
 			check(InputTensors.Num() == 0);
 			check(OutputTensors.Num() == 1);
@@ -42,7 +42,7 @@ namespace UE::NNERuntimeRDG::Private::Hlsl
 
 			switch (Attribute.GetType())
 			{
-				case ENNEAttributeDataType::Float:
+				case ENNERuntimeRDGDataAttributeDataType::Float:
 				{
 					if (Output.GetDataType() != ENNETensorDataType::Float)
 					{
@@ -56,7 +56,7 @@ namespace UE::NNERuntimeRDG::Private::Hlsl
 					Output.SetShape(FTensorShape::Make(Shape));
 					Output.SetPreparedData(Data);
 				} break;
-				case ENNEAttributeDataType::FloatArray:
+				case ENNERuntimeRDGDataAttributeDataType::FloatArray:
 				{
 					if (Output.GetDataType() != ENNETensorDataType::Float)
 					{
@@ -70,7 +70,7 @@ namespace UE::NNERuntimeRDG::Private::Hlsl
 					Output.SetShape(FTensorShape::Make(Shape));
 					Output.SetPreparedData(MakeArrayView(Values));
 				} break;
-				case ENNEAttributeDataType::Tensor:
+				case ENNERuntimeRDGDataAttributeDataType::Tensor:
 				{
 					const FAttributeTensor AttributeTensor = Attribute.GetValue<FAttributeTensor>();
 					if (Output.GetDataType() != AttributeTensor.GetDataType())
@@ -90,7 +90,7 @@ namespace UE::NNERuntimeRDG::Private::Hlsl
 			return 0;
 		};
 
-		virtual bool Initialize(TConstArrayView<NNE::FTensorDesc> InputTensorDescs, TConstArrayView<NNE::FTensorDesc> OutputTensorDescs, const NNE::FAttributeMap& Attributes) override
+		virtual bool Initialize(TConstArrayView<NNE::FTensorDesc> InputTensorDescs, TConstArrayView<NNE::FTensorDesc> OutputTensorDescs, const NNERuntimeRDGData::Internal::FAttributeMap& Attributes) override
 		{
 			check(InputTensorDescs.Num() == 0);
 			check(OutputTensorDescs.Num() == 1);
@@ -107,14 +107,14 @@ namespace UE::NNERuntimeRDG::Private::Hlsl
 		}
 	};
 
-	bool ValidateConstantOperator(const NNE::FAttributeMap& AttributeMap, TConstArrayView<ENNETensorDataType> InputTypes, TConstArrayView<NNE::FSymbolicTensorShape> InputShapes)
+	bool ValidateConstantOperator(const NNERuntimeRDGData::Internal::FAttributeMap& AttributeMap, TConstArrayView<ENNETensorDataType> InputTypes, TConstArrayView<NNE::FSymbolicTensorShape> InputShapes)
 	{
 		bool bIsValid = true;
 
 		FAttributeValidator AttributeValidator;
-		AttributeValidator.AddOptional(TEXT("value"), ENNEAttributeDataType::Tensor);
-		AttributeValidator.AddOptional(TEXT("value_float"), ENNEAttributeDataType::Float);
-		AttributeValidator.AddOptional(TEXT("value_floats"), ENNEAttributeDataType::FloatArray);
+		AttributeValidator.AddOptional(TEXT("value"), ENNERuntimeRDGDataAttributeDataType::Tensor);
+		AttributeValidator.AddOptional(TEXT("value_float"), ENNERuntimeRDGDataAttributeDataType::Float);
+		AttributeValidator.AddOptional(TEXT("value_floats"), ENNERuntimeRDGDataAttributeDataType::FloatArray);
 		if (!AttributeValidator.Validate(AttributeMap))
 		{
 			return false;
diff --git a/Engine/Plugins/Experimental/NNERuntimeRDG/Source/NNERuntimeRDG/Private/Hlsl/NNERuntimeRDGConv.cpp b/Engine/Plugins/Experimental/NNERuntimeRDG/Source/NNERuntimeRDG/Private/Hlsl/NNERuntimeRDGConv.cpp
--- a/Engine/Plugins/Experimental/NNERuntimeRDG/Source/NNERuntimeRDG/Private/Hlsl/NNERuntimeRDGConv.cpp
+++ b/Engine/Plugins/Experimental/NNERuntimeRDG/Source/NNERuntimeRDG/Private/Hlsl/NNERuntimeRDGConv.cpp
@@ -4,7 +4,6 @@
 
 #include "Algo/AnyOf.h"
 #include "HAL/IConsoleManager.h"
-#include "NNEAttributeMap.h"
 #include "NNETensor.h"
 #include "NNETypes.h"
 #include "NNEHlslShadersConvCS.h"
@@ -15,6 +14,7 @@
 #include "NNEHlslShadersConvWinogradWeightsCS.h"
 #include "NNEHlslShadersLog.h"
 #include "NNEHlslShadersTypeHelper.h"
+#include "NNERuntimeRDGDataAttributeMap.h"
 #include "NNERuntimeRDGHelperTranspose.h"
 #include "NNERuntimeRDGHlslHelper.h"
 
@@ -115,7 +115,7 @@ namespace UE::NNERuntimeRDG::Private::Hlsl
 			return 0;
 		};
 
-		virtual bool Initialize(TConstArrayView<NNE::FTensorDesc> InputTensorDescs, TConstArrayView<NNE::FTensorDesc> OutputTensorDescs, const NNE::FAttributeMap& Attributes) override
+		virtual bool Initialize(TConstArrayView<NNE::FTensorDesc> InputTensorDescs, TConstArrayView<NNE::FTensorDesc> OutputTensorDescs, const NNERuntimeRDGData::Internal::FAttributeMap& Attributes) override
 		{
             using namespace UE::NNEHlslShaders::Internal;
 
@@ -587,18 +587,18 @@ namespace UE::NNERuntimeRDG::Private::Hlsl
 		}
 	};
 
-	bool ValidateConvOperator(const NNE::FAttributeMap& AttributeMap, TConstArrayView<ENNETensorDataType> InputTypes, TConstArrayView<NNE::FSymbolicTensorShape> InputShapes)
+	bool ValidateConvOperator(const NNERuntimeRDGData::Internal::FAttributeMap& AttributeMap, TConstArrayView<ENNETensorDataType> InputTypes, TConstArrayView<NNE::FSymbolicTensorShape> InputShapes)
 	{
 		bool bIsValid = true;
 
 		FAttributeValidator AttributeValidator;
-		AttributeValidator.AddOptional(TEXT("auto_pad"), ENNEAttributeDataType::String);
-		AttributeValidator.AddOptional(TEXT("dilations"), ENNEAttributeDataType::Int32Array);
-		AttributeValidator.AddOptional(TEXT("group"), ENNEAttributeDataType::Int32);
-		AttributeValidator.AddOptional(TEXT("kernel_shape"), ENNEAttributeDataType::Int32Array); // idea: cross check input weight shape with this attribute if present
-		AttributeValidator.AddOptional(TEXT("pads"), ENNEAttributeDataType::Int32Array);
-		AttributeValidator.AddOptional(TEXT("strides"), ENNEAttributeDataType::Int32Array);
-		AttributeValidator.AddOptional(TEXT("__UE__WinogradPrecision"), ENNEAttributeDataType::Int32);
+		AttributeValidator.AddOptional(TEXT("auto_pad"), ENNERuntimeRDGDataAttributeDataType::String);
+		AttributeValidator.AddOptional(TEXT("dilations"), ENNERuntimeRDGDataAttributeDataType::Int32Array);
+		AttributeValidator.AddOptional(TEXT("group"), ENNERuntimeRDGDataAttributeDataType::Int32);
+		AttributeValidator.AddOptional(TEXT("kernel_shape"), ENNERuntimeRDGDataAttributeDataType::Int32Array); // idea: cross check input weight shape with this attribute if present
+		AttributeValidator.AddOptional(TEXT("pads"), ENNERuntimeRDGDataAttributeDataType::Int32Array);
+		AttributeValidator.AddOptional(TEXT("strides"), ENNERuntimeRDGDataAttributeDataType::Int32Array);
+		AttributeValidator.AddOptional(TEXT("__UE__WinogradPrecision"), ENNERuntimeRDGDataAttributeDataType::Int32);
 
 		bIsValid &= AttributeValidator.Validate(AttributeMap);
 
diff --git a/Engine/Plugins/Experimental/NNERuntimeRDG/Source/NNERuntimeRDG/Private/Hlsl/NNERuntimeRDGConvTranspose.cpp b/Engine/Plugins/Experimental/NNERuntimeRDG/Source/NNERuntimeRDG/Private/Hlsl/NNERuntimeRDGConvTranspose.cpp
--- a/Engine/Plugins/Experimental/NNERuntimeRDG/Source/NNERuntimeRDG/Private/Hlsl/NNERuntimeRDGConvTranspose.cpp
+++ b/Engine/Plugins/Experimental/NNERuntimeRDG/Source/NNERuntimeRDG/Private/Hlsl/NNERuntimeRDGConvTranspose.cpp
@@ -63,7 +63,7 @@ namespace UE::NNERuntimeRDG::Private::Hlsl
 			return 0;
 		};
 
-		virtual bool Initialize(TConstArrayView<NNE::FTensorDesc> InputTensorDescs, TConstArrayView<NNE::FTensorDesc> OutputTensorDescs, const NNE::FAttributeMap& Attributes) override
+		virtual bool Initialize(TConstArrayView<NNE::FTensorDesc> InputTensorDescs, TConstArrayView<NNE::FTensorDesc> OutputTensorDescs, const NNERuntimeRDGData::Internal::FAttributeMap& Attributes) override
 		{
 			using namespace UE::NNEHlslShaders::Internal;
 
@@ -179,19 +179,19 @@ namespace UE::NNERuntimeRDG::Private::Hlsl
 		}
 	};
 
-	bool ValidateConvTransposeOperator(const NNE::FAttributeMap& AttributeMap, TConstArrayView<ENNETensorDataType> InputTypes, TConstArrayView<NNE::FSymbolicTensorShape> InputShapes)
+	bool ValidateConvTransposeOperator(const NNERuntimeRDGData::Internal::FAttributeMap& AttributeMap, TConstArrayView<ENNETensorDataType> InputTypes, TConstArrayView<NNE::FSymbolicTensorShape> InputShapes)
 	{
 		bool bIsValid = true;
 
 		FAttributeValidator AttributeValidator;
-		AttributeValidator.AddOptional(TEXT("auto_pad"), ENNEAttributeDataType::String);
-		AttributeValidator.AddOptional(TEXT("dilations"), ENNEAttributeDataType::Int32Array);
-		AttributeValidator.AddOptional(TEXT("group"), ENNEAttributeDataType::Int32);
-		//AttributeValidator.AddOptional(TEXT("kernel_shape"), ENNEAttributeDataType::Int32Array);
-		AttributeValidator.AddOptional(TEXT("output_padding"), ENNEAttributeDataType::Int32Array);
-		//AttributeValidator.AddOptional(TEXT("output_shape"), ENNEAttributeDataType::Int32Array);
-		AttributeValidator.AddOptional(TEXT("pads"), ENNEAttributeDataType::Int32Array);
-		AttributeValidator.AddOptional(TEXT("strides"), ENNEAttributeDataType::Int32Array);
+		AttributeValidator.AddOptional(TEXT("auto_pad"), ENNERuntimeRDGDataAttributeDataType::String);
+		AttributeValidator.AddOptional(TEXT("dilations"), ENNERuntimeRDGDataAttributeDataType::Int32Array);
+		AttributeValidator.AddOptional(TEXT("group"), ENNERuntimeRDGDataAttributeDataType::Int32);
+		//AttributeValidator.AddOptional(TEXT("kernel_shape"), ENNERuntimeRDGDataAttributeDataType::Int32Array);
+		AttributeValidator.AddOptional(TEXT("output_padding"), ENNERuntimeRDGDataAttributeDataType::Int32Array);
+		//AttributeValidator.AddOptional(TEXT("output_shape"), ENNERuntimeRDGDataAttributeDataType::Int32Array);
+		AttributeValidator.AddOptional(TEXT("pads"), ENNERuntimeRDGDataAttributeDataType::Int32Array);
+		AttributeValidator.AddOptional(TEXT("strides"), ENNERuntimeRDGDataAttributeDataType::Int32Array);
 
 		bIsValid &= AttributeValidator.Validate(AttributeMap);
 
diff --git a/Engine/Plugins/Experimental/NNERuntimeRDG/Source/NNERuntimeRDG/Private/Hlsl/NNERuntimeRDGCumSum.cpp b/Engine/Plugins/Experimental/NNERuntimeRDG/Source/NNERuntimeRDG/Private/Hlsl/NNERuntimeRDGCumSum.cpp
--- a/Engine/Plugins/Experimental/NNERuntimeRDG/Source/NNERuntimeRDG/Private/Hlsl/NNERuntimeRDGCumSum.cpp
+++ b/Engine/Plugins/Experimental/NNERuntimeRDG/Source/NNERuntimeRDG/Private/Hlsl/NNERuntimeRDGCumSum.cpp
@@ -127,7 +127,7 @@ namespace UE::NNERuntimeRDG::Private::Hlsl
 			return 0;
 		};
 
-		virtual bool Initialize(TConstArrayView<NNE::FTensorDesc> InputTensorDescs, TConstArrayView<NNE::FTensorDesc> OutputTensorDescs, const NNE::FAttributeMap& Attributes) override
+		virtual bool Initialize(TConstArrayView<NNE::FTensorDesc> InputTensorDescs, TConstArrayView<NNE::FTensorDesc> OutputTensorDescs, const NNERuntimeRDGData::Internal::FAttributeMap& Attributes) override
 		{
 			check(InputTensorDescs.Num() == 2);
 			check(OutputTensorDescs.Num() == 1);
@@ -245,7 +245,7 @@ namespace UE::NNERuntimeRDG::Private::Hlsl
 	};
 
 	template<int Version>
-	bool ValidateCumSumOperator(const NNE::FAttributeMap& AttributeMap, TConstArrayView<ENNETensorDataType> InputTypes, TConstArrayView<NNE::FSymbolicTensorShape> InputShapes)
+	bool ValidateCumSumOperator(const NNERuntimeRDGData::Internal::FAttributeMap& AttributeMap, TConstArrayView<ENNETensorDataType> InputTypes, TConstArrayView<NNE::FSymbolicTensorShape> InputShapes)
 	{
 		bool bIsValid = true;
 
@@ -255,8 +255,8 @@ namespace UE::NNERuntimeRDG::Private::Hlsl
 		// [REDACTED_PROFILE_URL]/onnx/blob/main/docs/Changelog.md#CumSum-14
 
 		FAttributeValidator AttributeValidator;
-		AttributeValidator.AddOptional(TEXT("exclusive"), ENNEAttributeDataType::Int32);
-		AttributeValidator.AddOptional(TEXT("reverse"), ENNEAttributeDataType::Int32);
+		AttributeValidator.AddOptional(TEXT("exclusive"), ENNERuntimeRDGDataAttributeDataType::Int32);
+		AttributeValidator.AddOptional(TEXT("reverse"), ENNERuntimeRDGDataAttributeDataType::Int32);
 		
 		bIsValid &= AttributeValidator.Validate(AttributeMap);
 
diff --git a/Engine/Plugins/Experimental/NNERuntimeRDG/Source/NNERuntimeRDG/Private/Hlsl/NNERuntimeRDGDepthToSpace.cpp b/Engine/Plugins/Experimental/NNERuntimeRDG/Source/NNERuntimeRDG/Private/Hlsl/NNERuntimeRDGDepthToSpace.cpp
--- a/Engine/Plugins/Experimental/NNERuntimeRDG/Source/NNERuntimeRDG/Private/Hlsl/NNERuntimeRDGDepthToSpace.cpp
+++ b/Engine/Plugins/Experimental/NNERuntimeRDG/Source/NNERuntimeRDG/Private/Hlsl/NNERuntimeRDGDepthToSpace.cpp
@@ -89,7 +89,7 @@ namespace UE::NNERuntimeRDG::Private::Hlsl
 			return 0;
 		};
 
-		virtual bool Initialize(TConstArrayView<NNE::FTensorDesc> InputTensorDescs, TConstArrayView<NNE::FTensorDesc> OutputTensorDescs, const NNE::FAttributeMap& Attributes) override
+		virtual bool Initialize(TConstArrayView<NNE::FTensorDesc> InputTensorDescs, TConstArrayView<NNE::FTensorDesc> OutputTensorDescs, const NNERuntimeRDGData::Internal::FAttributeMap& Attributes) override
 		{
 			check(InputTensorDescs.Num() == 1);
 			check(OutputTensorDescs.Num() == 1);
@@ -97,7 +97,7 @@ namespace UE::NNERuntimeRDG::Private::Hlsl
 			const NNE::FTensorDesc& InputDesc = InputTensorDescs[0];
 			const int32 InputRank = InputDesc.GetShape().Rank();
 
-			const FNNEAttributeValue* BlockSizeAttr = Attributes.GetAttributeValue(TEXT("blocksize"));
+			const FNNERuntimeRDGDataAttributeValue* BlockSizeAttr = Attributes.GetAttributeValue(TEXT("blocksize"));
 			if(BlockSizeAttr == nullptr)
 			{
 				UE_LOG(LogNNERuntimeRDGHlsl, Warning, TEXT("DepthToSpace: blocksize attribute is required."));
@@ -177,13 +177,13 @@ namespace UE::NNERuntimeRDG::Private::Hlsl
 		}
 	};
 
-	bool ValidateDepthToSpaceOperator(const NNE::FAttributeMap& AttributeMap, TConstArrayView<ENNETensorDataType> InputTypes, TConstArrayView<NNE::FSymbolicTensorShape> InputShapes)
+	bool ValidateDepthToSpaceOperator(const NNERuntimeRDGData::Internal::FAttributeMap& AttributeMap, TConstArrayView<ENNETensorDataType> InputTypes, TConstArrayView<NNE::FSymbolicTensorShape> InputShapes)
 	{
 		bool bIsValid = true;
 
 		FAttributeValidator AttributeValidator;
-		AttributeValidator.AddRequired(TEXT("blocksize"), ENNEAttributeDataType::Int32);
-		AttributeValidator.AddOptional(TEXT("mode"), ENNEAttributeDataType::String);
+		AttributeValidator.AddRequired(TEXT("blocksize"), ENNERuntimeRDGDataAttributeDataType::Int32);
+		AttributeValidator.AddOptional(TEXT("mode"), ENNERuntimeRDGDataAttributeDataType::String);
 		bIsValid &= AttributeValidator.Validate(AttributeMap);
 
 		FInputValidator InputValidator;
diff --git a/Engine/Plugins/Experimental/NNERuntimeRDG/Source/NNERuntimeRDG/Private/Hlsl/NNERuntimeRDGDropout.cpp b/Engine/Plugins/Experimental/NNERuntimeRDG/Source/NNERuntimeRDG/Private/Hlsl/NNERuntimeRDGDropout.cpp
--- a/Engine/Plugins/Experimental/NNERuntimeRDG/Source/NNERuntimeRDG/Private/Hlsl/NNERuntimeRDGDropout.cpp
+++ b/Engine/Plugins/Experimental/NNERuntimeRDG/Source/NNERuntimeRDG/Private/Hlsl/NNERuntimeRDGDropout.cpp
@@ -39,7 +39,7 @@ namespace UE::NNERuntimeRDG::Private::Hlsl
 			return 0;
 		};
 
-		virtual bool Initialize(TConstArrayView<NNE::FTensorDesc> InputTensorDescs, TConstArrayView<NNE::FTensorDesc> OutputTensorDescs, const NNE::FAttributeMap& Attributes) override
+		virtual bool Initialize(TConstArrayView<NNE::FTensorDesc> InputTensorDescs, TConstArrayView<NNE::FTensorDesc> OutputTensorDescs, const NNERuntimeRDGData::Internal::FAttributeMap& Attributes) override
 		{
 			check(InputTensorDescs.Num() == 1);
 			check(OutputTensorDescs.Num() == 1 || OutputTensorDescs.Num() == 2);
@@ -71,19 +71,19 @@ namespace UE::NNERuntimeRDG::Private::Hlsl
 	};
 
 	template<int Version>
-	bool ValidateDropoutOperator(const NNE::FAttributeMap& AttributeMap, TConstArrayView<ENNETensorDataType> InputTypes, TConstArrayView<NNE::FSymbolicTensorShape> InputDropouts)
+	bool ValidateDropoutOperator(const NNERuntimeRDGData::Internal::FAttributeMap& AttributeMap, TConstArrayView<ENNETensorDataType> InputTypes, TConstArrayView<NNE::FSymbolicTensorShape> InputDropouts)
 	{
 		bool bIsValid = true;
 
 		FAttributeValidator AttributeValidator;
 		if constexpr (Version >= 12)
 		{
-			AttributeValidator.AddOptional(TEXT("seed"), ENNEAttributeDataType::Int32);
+			AttributeValidator.AddOptional(TEXT("seed"), ENNERuntimeRDGDataAttributeDataType::Int32);
 		}
 		else
 		{
 			static_assert(Version >= 7, "Minimum supported version for operator Dropout is 7!");
-			AttributeValidator.AddOptional(TEXT("ratio"), ENNEAttributeDataType::Float);
+			AttributeValidator.AddOptional(TEXT("ratio"), ENNERuntimeRDGDataAttributeDataType::Float);
 		}
 		bIsValid &= AttributeValidator.Validate(AttributeMap);
 
diff --git a/Engine/Plugins/Experimental/NNERuntimeRDG/Source/NNERuntimeRDG/Private/Hlsl/NNERuntimeRDGElementWiseBinary.cpp b/Engine/Plugins/Experimental/NNERuntimeRDG/Source/NNERuntimeRDG/Private/Hlsl/NNERuntimeRDGElementWiseBinary.cpp
--- a/Engine/Plugins/Experimental/NNERuntimeRDG/Source/NNERuntimeRDG/Private/Hlsl/NNERuntimeRDGElementWiseBinary.cpp
+++ b/Engine/Plugins/Experimental/NNERuntimeRDG/Source/NNERuntimeRDG/Private/Hlsl/NNERuntimeRDGElementWiseBinary.cpp
@@ -72,7 +72,7 @@ namespace UE::NNERuntimeRDG::Private::Hlsl
 			return 0;
 		}
 		
-		virtual bool Initialize(TConstArrayView<NNE::FTensorDesc> InputTensorDescs, TConstArrayView<NNE::FTensorDesc> OutputTensorDescs, const NNE::FAttributeMap& Attributes) override
+		virtual bool Initialize(TConstArrayView<NNE::FTensorDesc> InputTensorDescs, TConstArrayView<NNE::FTensorDesc> OutputTensorDescs, const NNERuntimeRDGData::Internal::FAttributeMap& Attributes) override
 		{
 			check(InputTensorDescs.Num() == 2);
 			check(OutputTensorDescs.Num() == 1);
@@ -130,7 +130,7 @@ namespace UE::NNERuntimeRDG::Private::Hlsl
 		}
 	};
 
-	bool ValidateElementWiseBinaryOperator(const NNE::FAttributeMap& AttributeMap, TConstArrayView<ENNETensorDataType> InputTypes, TConstArrayView<NNE::FSymbolicTensorShape> InputShapes)
+	bool ValidateElementWiseBinaryOperator(const NNERuntimeRDGData::Internal::FAttributeMap& AttributeMap, TConstArrayView<ENNETensorDataType> InputTypes, TConstArrayView<NNE::FSymbolicTensorShape> InputShapes)
 	{
 		bool bIsValid = true;
 
diff --git a/Engine/Plugins/Experimental/NNERuntimeRDG/Source/NNERuntimeRDG/Private/Hlsl/NNERuntimeRDGElementWiseUnary.cpp b/Engine/Plugins/Experimental/NNERuntimeRDG/Source/NNERuntimeRDG/Private/Hlsl/NNERuntimeRDGElementWiseUnary.cpp
--- a/Engine/Plugins/Experimental/NNERuntimeRDG/Source/NNERuntimeRDG/Private/Hlsl/NNERuntimeRDGElementWiseUnary.cpp
+++ b/Engine/Plugins/Experimental/NNERuntimeRDG/Source/NNERuntimeRDG/Private/Hlsl/NNERuntimeRDGElementWiseUnary.cpp
@@ -2,9 +2,9 @@
 
 #include "NNERuntimeRDGElementWiseUnary.h"
 
-#include "NNEAttributeMap.h"
 #include "NNEHlslShadersElementWiseUnaryCS.h"
 #include "NNEHlslShadersTypeHelper.h"
+#include "NNERuntimeRDGDataAttributeMap.h"
 #include "NNERuntimeRDGHelperElementWiseUnary.h"
 #include "NNERuntimeRDGHlslHelper.h"
 #include "NNETensor.h"
@@ -73,7 +73,7 @@ namespace UE::NNERuntimeRDG::Private::Hlsl
 			return 0;
 		}
 
-		virtual bool Initialize(TConstArrayView<NNE::FTensorDesc> InputTensorDescs, TConstArrayView<NNE::FTensorDesc> OutputTensorDescs, const NNE::FAttributeMap& Attributes) override
+		virtual bool Initialize(TConstArrayView<NNE::FTensorDesc> InputTensorDescs, TConstArrayView<NNE::FTensorDesc> OutputTensorDescs, const NNERuntimeRDGData::Internal::FAttributeMap& Attributes) override
 		{
 			check(InputTensorDescs.Num() == 1);
 			check(OutputTensorDescs.Num() == 1);
@@ -171,7 +171,7 @@ namespace UE::NNERuntimeRDG::Private::Hlsl
 	{
 	}
 
-	template<> bool TElementWiseUnary<EElementWiseUnaryOperatorType::Clip>::Initialize(TConstArrayView<NNE::FTensorDesc> InputTensorDescs, TConstArrayView<NNE::FTensorDesc> OutputTensorDescs, const NNE::FAttributeMap& Attributes)
+	template<> bool TElementWiseUnary<EElementWiseUnaryOperatorType::Clip>::Initialize(TConstArrayView<NNE::FTensorDesc> InputTensorDescs, TConstArrayView<NNE::FTensorDesc> OutputTensorDescs, const NNERuntimeRDGData::Internal::FAttributeMap& Attributes)
 	{
 		check(Version >= 11 ? (InputTensorDescs.Num() >= 1 && InputTensorDescs.Num() <= 3) : InputTensorDescs.Num() == 1);
 		check(OutputTensorDescs.Num() == 1);
@@ -188,7 +188,7 @@ namespace UE::NNERuntimeRDG::Private::Hlsl
 	}
 	
 	template<EElementWiseUnaryOperatorType OpType>
-	bool ValidateElementWiseUnaryOperatorImpl(int Version, const NNE::FAttributeMap& AttributeMap, TConstArrayView<ENNETensorDataType> InputTypes, TConstArrayView<NNE::FSymbolicTensorShape> InputShapes)
+	bool ValidateElementWiseUnaryOperatorImpl(int Version, const NNERuntimeRDGData::Internal::FAttributeMap& AttributeMap, TConstArrayView<ENNETensorDataType> InputTypes, TConstArrayView<NNE::FSymbolicTensorShape> InputShapes)
 	{
 		bool bIsValid = true;
 
@@ -205,13 +205,13 @@ namespace UE::NNERuntimeRDG::Private::Hlsl
 	}
 
 	template<>
-	bool ValidateElementWiseUnaryOperatorImpl<EElementWiseUnaryOperatorType::Selu>(int Version, const NNE::FAttributeMap& AttributeMap, TConstArrayView<ENNETensorDataType> InputTypes, TConstArrayView<NNE::FSymbolicTensorShape> InputShapes)
+	bool ValidateElementWiseUnaryOperatorImpl<EElementWiseUnaryOperatorType::Selu>(int Version, const NNERuntimeRDGData::Internal::FAttributeMap& AttributeMap, TConstArrayView<ENNETensorDataType> InputTypes, TConstArrayView<NNE::FSymbolicTensorShape> InputShapes)
 	{
 		bool bIsValid = true;
 
 		FAttributeValidator AttributeValidator;
-		AttributeValidator.AddOptional(TEXT("alpha"), ENNEAttributeDataType::Float);
-		AttributeValidator.AddOptional(TEXT("gamma"), ENNEAttributeDataType::Float);
+		AttributeValidator.AddOptional(TEXT("alpha"), ENNERuntimeRDGDataAttributeDataType::Float);
+		AttributeValidator.AddOptional(TEXT("gamma"), ENNERuntimeRDGDataAttributeDataType::Float);
 		bIsValid &= AttributeValidator.Validate(AttributeMap);
 
 		FInputValidator InputValidator;
@@ -224,12 +224,12 @@ namespace UE::NNERuntimeRDG::Private::Hlsl
 	}
 
 	template<>
-	bool ValidateElementWiseUnaryOperatorImpl<EElementWiseUnaryOperatorType::Elu>(int Version, const NNE::FAttributeMap& AttributeMap, TConstArrayView<ENNETensorDataType> InputTypes, TConstArrayView<NNE::FSymbolicTensorShape> InputShapes)
+	bool ValidateElementWiseUnaryOperatorImpl<EElementWiseUnaryOperatorType::Elu>(int Version, const NNERuntimeRDGData::Internal::FAttributeMap& AttributeMap, TConstArrayView<ENNETensorDataType> InputTypes, TConstArrayView<NNE::FSymbolicTensorShape> InputShapes)
 	{
 		bool bIsValid = true;
 
 		FAttributeValidator AttributeValidator;
-		AttributeValidator.AddOptional(TEXT("alpha"), ENNEAttributeDataType::Float);
+		AttributeValidator.AddOptional(TEXT("alpha"), ENNERuntimeRDGDataAttributeDataType::Float);
 		bIsValid &= AttributeValidator.Validate(AttributeMap);
 
 		FInputValidator InputValidator;
@@ -242,13 +242,13 @@ namespace UE::NNERuntimeRDG::Private::Hlsl
 	}
 
 	template<>
-	bool ValidateElementWiseUnaryOperatorImpl<EElementWiseUnaryOperatorType::HardSigmoid>(int Version, const NNE::FAttributeMap& AttributeMap, TConstArrayView<ENNETensorDataType> InputTypes, TConstArrayView<NNE::FSymbolicTensorShape> InputShapes)
+	bool ValidateElementWiseUnaryOperatorImpl<EElementWiseUnaryOperatorType::HardSigmoid>(int Version, const NNERuntimeRDGData::Internal::FAttributeMap& AttributeMap, TConstArrayView<ENNETensorDataType> InputTypes, TConstArrayView<NNE::FSymbolicTensorShape> InputShapes)
 	{
 		bool bIsValid = true;
 
 		FAttributeValidator AttributeValidator;
-		AttributeValidator.AddOptional(TEXT("alpha"), ENNEAttributeDataType::Float);
-		AttributeValidator.AddOptional(TEXT("beta"), ENNEAttributeDataType::Float);
+		AttributeValidator.AddOptional(TEXT("alpha"), ENNERuntimeRDGDataAttributeDataType::Float);
+		AttributeValidator.AddOptional(TEXT("beta"), ENNERuntimeRDGDataAttributeDataType::Float);
 		bIsValid &= AttributeValidator.Validate(AttributeMap);
 
 		FInputValidator InputValidator;
@@ -261,12 +261,12 @@ namespace UE::NNERuntimeRDG::Private::Hlsl
 	}
 
 	template<>
-	bool ValidateElementWiseUnaryOperatorImpl<EElementWiseUnaryOperatorType::LeakyRelu>(int Version, const NNE::FAttributeMap& AttributeMap, TConstArrayView<ENNETensorDataType> InputTypes, TConstArrayView<NNE::FSymbolicTensorShape> InputShapes)
+	bool ValidateElementWiseUnaryOperatorImpl<EElementWiseUnaryOperatorType::LeakyRelu>(int Version, const NNERuntimeRDGData::Internal::FAttributeMap& AttributeMap, TConstArrayView<ENNETensorDataType> InputTypes, TConstArrayView<NNE::FSymbolicTensorShape> InputShapes)
 	{
 		bool bIsValid = true;
 
 		FAttributeValidator AttributeValidator;
-		AttributeValidator.AddOptional(TEXT("alpha"), ENNEAttributeDataType::Float);
+		AttributeValidator.AddOptional(TEXT("alpha"), ENNERuntimeRDGDataAttributeDataType::Float);
 		bIsValid &= AttributeValidator.Validate(AttributeMap);
 
 		FInputValidator InputValidator;
@@ -279,15 +279,15 @@ namespace UE::NNERuntimeRDG::Private::Hlsl
 	}
 
 	template<>
-	bool ValidateElementWiseUnaryOperatorImpl<EElementWiseUnaryOperatorType::Clip>(int Version, const NNE::FAttributeMap& AttributeMap, TConstArrayView<ENNETensorDataType> InputTypes, TConstArrayView<NNE::FSymbolicTensorShape> InputShapes)
+	bool ValidateElementWiseUnaryOperatorImpl<EElementWiseUnaryOperatorType::Clip>(int Version, const NNERuntimeRDGData::Internal::FAttributeMap& AttributeMap, TConstArrayView<ENNETensorDataType> InputTypes, TConstArrayView<NNE::FSymbolicTensorShape> InputShapes)
 	{
 		bool bIsValid = true;
 
 		FAttributeValidator AttributeValidator;
 		if(Version < 11)
 		{
-			AttributeValidator.AddOptional(TEXT("min"), ENNEAttributeDataType::Float);
-			AttributeValidator.AddOptional(TEXT("max"), ENNEAttributeDataType::Float);
+			AttributeValidator.AddOptional(TEXT("min"), ENNERuntimeRDGDataAttributeDataType::Float);
+			AttributeValidator.AddOptional(TEXT("max"), ENNERuntimeRDGDataAttributeDataType::Float);
 		}
 		bIsValid &= AttributeValidator.Validate(AttributeMap);
 
@@ -306,7 +306,7 @@ namespace UE::NNERuntimeRDG::Private::Hlsl
 	}
 
 	template<EElementWiseUnaryOperatorType OpType, int Version>
-	bool ValidateElementWiseUnaryOperator(const NNE::FAttributeMap& AttributeMap, TConstArrayView<ENNETensorDataType> InputTypes, TConstArrayView<NNE::FSymbolicTensorShape> InputShapes)
+	bool ValidateElementWiseUnaryOperator(const NNERuntimeRDGData::Internal::FAttributeMap& AttributeMap, TConstArrayView<ENNETensorDataType> InputTypes, TConstArrayView<NNE::FSymbolicTensorShape> InputShapes)
 	{
 		return ValidateElementWiseUnaryOperatorImpl<OpType>(Version, AttributeMap, InputTypes, InputShapes);
 	}
diff --git a/Engine/Plugins/Experimental/NNERuntimeRDG/Source/NNERuntimeRDG/Private/Hlsl/NNERuntimeRDGElementWiseVariadic.cpp b/Engine/Plugins/Experimental/NNERuntimeRDG/Source/NNERuntimeRDG/Private/Hlsl/NNERuntimeRDGElementWiseVariadic.cpp
--- a/Engine/Plugins/Experimental/NNERuntimeRDG/Source/NNERuntimeRDG/Private/Hlsl/NNERuntimeRDGElementWiseVariadic.cpp
+++ b/Engine/Plugins/Experimental/NNERuntimeRDG/Source/NNERuntimeRDG/Private/Hlsl/NNERuntimeRDGElementWiseVariadic.cpp
@@ -137,7 +137,7 @@ namespace UE::NNERuntimeRDG::Private::Hlsl
 			return 0;
 		}
 		
-		virtual bool Initialize(TConstArrayView<NNE::FTensorDesc> InputTensorDescs, TConstArrayView<NNE::FTensorDesc> OutputTensorDescs, const NNE::FAttributeMap& Attributes) override
+		virtual bool Initialize(TConstArrayView<NNE::FTensorDesc> InputTensorDescs, TConstArrayView<NNE::FTensorDesc> OutputTensorDescs, const NNERuntimeRDGData::Internal::FAttributeMap& Attributes) override
 		{
 			check(InputTensorDescs.Num() > 0);
 			check(OutputTensorDescs.Num() == 1);
@@ -182,7 +182,7 @@ namespace UE::NNERuntimeRDG::Private::Hlsl
 		}
 	};
 
-	bool ValidateElementWiseVariadicOperator(const NNE::FAttributeMap& AttributeMap, TConstArrayView<ENNETensorDataType> InputTypes, TConstArrayView<NNE::FSymbolicTensorShape> InputShapes)
+	bool ValidateElementWiseVariadicOperator(const NNERuntimeRDGData::Internal::FAttributeMap& AttributeMap, TConstArrayView<ENNETensorDataType> InputTypes, TConstArrayView<NNE::FSymbolicTensorShape> InputShapes)
 	{
 		bool bIsValid = true;
 
diff --git a/Engine/Plugins/Experimental/NNERuntimeRDG/Source/NNERuntimeRDG/Private/Hlsl/NNERuntimeRDGFlatten.cpp b/Engine/Plugins/Experimental/NNERuntimeRDG/Source/NNERuntimeRDG/Private/Hlsl/NNERuntimeRDGFlatten.cpp
--- a/Engine/Plugins/Experimental/NNERuntimeRDG/Source/NNERuntimeRDG/Private/Hlsl/NNERuntimeRDGFlatten.cpp
+++ b/Engine/Plugins/Experimental/NNERuntimeRDG/Source/NNERuntimeRDG/Private/Hlsl/NNERuntimeRDGFlatten.cpp
@@ -52,7 +52,7 @@ namespace UE::NNERuntimeRDG::Private::Hlsl
 			return 0;
 		};
 
-		virtual bool Initialize(TConstArrayView<NNE::FTensorDesc> InputTensorDescs, TConstArrayView<NNE::FTensorDesc> OutputTensorDescs, const NNE::FAttributeMap& Attributes) override
+		virtual bool Initialize(TConstArrayView<NNE::FTensorDesc> InputTensorDescs, TConstArrayView<NNE::FTensorDesc> OutputTensorDescs, const NNERuntimeRDGData::Internal::FAttributeMap& Attributes) override
 		{
 			check(InputTensorDescs.Num() == 1);
 			check(OutputTensorDescs.Num() == 1);
@@ -90,12 +90,12 @@ namespace UE::NNERuntimeRDG::Private::Hlsl
 		}
 	};
 
-	bool ValidateFlattenOperator(const NNE::FAttributeMap& AttributeMap, TConstArrayView<ENNETensorDataType> InputTypes, TConstArrayView<NNE::FSymbolicTensorShape> InputFlattens)
+	bool ValidateFlattenOperator(const NNERuntimeRDGData::Internal::FAttributeMap& AttributeMap, TConstArrayView<ENNETensorDataType> InputTypes, TConstArrayView<NNE::FSymbolicTensorShape> InputFlattens)
 	{
 		bool bIsValid = true;
 
 		FAttributeValidator AttributeValidator;
-		AttributeValidator.AddOptional(TEXT("axis"), ENNEAttributeDataType::Int32);
+		AttributeValidator.AddOptional(TEXT("axis"), ENNERuntimeRDGDataAttributeDataType::Int32);
 		bIsValid &= AttributeValidator.Validate(AttributeMap);
 
 		FInputValidator InputValidator;
diff --git a/Engine/Plugins/Experimental/NNERuntimeRDG/Source/NNERuntimeRDG/Private/Hlsl/NNERuntimeRDGGather.cpp b/Engine/Plugins/Experimental/NNERuntimeRDG/Source/NNERuntimeRDG/Private/Hlsl/NNERuntimeRDGGather.cpp
--- a/Engine/Plugins/Experimental/NNERuntimeRDG/Source/NNERuntimeRDG/Private/Hlsl/NNERuntimeRDGGather.cpp
+++ b/Engine/Plugins/Experimental/NNERuntimeRDG/Source/NNERuntimeRDG/Private/Hlsl/NNERuntimeRDGGather.cpp
@@ -2,9 +2,9 @@
 
 #include "NNERuntimeRDGGather.h"
 
-#include "NNEAttributeMap.h"
 #include "NNEHlslShadersGatherCS.h"
 #include "NNEHlslShadersLog.h"
+#include "NNERuntimeRDGDataAttributeMap.h"
 #include "NNERuntimeRDGHelperGather.h"
 #include "NNERuntimeRDGHlslHelper.h"
 #include "NNETypes.h"
@@ -70,7 +70,7 @@ namespace UE::NNERuntimeRDG::Private::Hlsl
 			return 0;
 		};
 		
-		virtual bool Initialize(TConstArrayView<NNE::FTensorDesc> InputTensorDescs, TConstArrayView<NNE::FTensorDesc> OutputTensorDescs, const NNE::FAttributeMap& Attributes) override
+		virtual bool Initialize(TConstArrayView<NNE::FTensorDesc> InputTensorDescs, TConstArrayView<NNE::FTensorDesc> OutputTensorDescs, const NNERuntimeRDGData::Internal::FAttributeMap& Attributes) override
 		{
 			check(InputTensorDescs.Num() == 2)
 			check(OutputTensorDescs.Num() == 1)
@@ -141,12 +141,12 @@ namespace UE::NNERuntimeRDG::Private::Hlsl
 		}
 	};
 
-	bool ValidateGatherOperator(const NNE::FAttributeMap& AttributeMap, TConstArrayView<ENNETensorDataType> InputTypes, TConstArrayView<NNE::FSymbolicTensorShape> InputShapes)
+	bool ValidateGatherOperator(const NNERuntimeRDGData::Internal::FAttributeMap& AttributeMap, TConstArrayView<ENNETensorDataType> InputTypes, TConstArrayView<NNE::FSymbolicTensorShape> InputShapes)
 	{
 		bool bIsValid = true;
 
 		FAttributeValidator AttributeValidator;
-		AttributeValidator.AddOptional(TEXT("axis"), ENNEAttributeDataType::Int32);
+		AttributeValidator.AddOptional(TEXT("axis"), ENNERuntimeRDGDataAttributeDataType::Int32);
 		bIsValid &= AttributeValidator.Validate(AttributeMap);
 
 		FInputValidator InputValidator;
diff --git a/Engine/Plugins/Experimental/NNERuntimeRDG/Source/NNERuntimeRDG/Private/Hlsl/NNERuntimeRDGGatherElements.cpp b/Engine/Plugins/Experimental/NNERuntimeRDG/Source/NNERuntimeRDG/Private/Hlsl/NNERuntimeRDGGatherElements.cpp
--- a/Engine/Plugins/Experimental/NNERuntimeRDG/Source/NNERuntimeRDG/Private/Hlsl/NNERuntimeRDGGatherElements.cpp
+++ b/Engine/Plugins/Experimental/NNERuntimeRDG/Source/NNERuntimeRDG/Private/Hlsl/NNERuntimeRDGGatherElements.cpp
@@ -38,7 +38,7 @@ namespace UE::NNERuntimeRDG::Private::Hlsl
 			return 0;
 		}
 
-		virtual bool Initialize(TConstArrayView<NNE::FTensorDesc> InputTensorDescs, TConstArrayView<NNE::FTensorDesc> OutputTensorDescs, const NNE::FAttributeMap& Attributes) override
+		virtual bool Initialize(TConstArrayView<NNE::FTensorDesc> InputTensorDescs, TConstArrayView<NNE::FTensorDesc> OutputTensorDescs, const NNERuntimeRDGData::Internal::FAttributeMap& Attributes) override
 		{
 			using namespace UE::NNEHlslShaders::Internal;
 
@@ -124,12 +124,12 @@ namespace UE::NNERuntimeRDG::Private::Hlsl
 	};
 
 	template<int Version>
-	bool ValidateGatherElementsOperator(const NNE::FAttributeMap& AttributeMap, TConstArrayView<ENNETensorDataType> InputTypes, TConstArrayView<NNE::FSymbolicTensorShape> InputShapes)
+	bool ValidateGatherElementsOperator(const NNERuntimeRDGData::Internal::FAttributeMap& AttributeMap, TConstArrayView<ENNETensorDataType> InputTypes, TConstArrayView<NNE::FSymbolicTensorShape> InputShapes)
 	{
 		bool bIsValid = true;
 
 		FAttributeValidator AttributeValidator;
-		AttributeValidator.AddOptional(TEXT("axis"), ENNEAttributeDataType::Int32);
+		AttributeValidator.AddOptional(TEXT("axis"), ENNERuntimeRDGDataAttributeDataType::Int32);
 		bIsValid &= AttributeValidator.Validate(AttributeMap);
 		
 		FInputValidator InputValidator;
diff --git a/Engine/Plugins/Experimental/NNERuntimeRDG/Source/NNERuntimeRDG/Private/Hlsl/NNERuntimeRDGGemm.cpp b/Engine/Plugins/Experimental/NNERuntimeRDG/Source/NNERuntimeRDG/Private/Hlsl/NNERuntimeRDGGemm.cpp
--- a/Engine/Plugins/Experimental/NNERuntimeRDG/Source/NNERuntimeRDG/Private/Hlsl/NNERuntimeRDGGemm.cpp
+++ b/Engine/Plugins/Experimental/NNERuntimeRDG/Source/NNERuntimeRDG/Private/Hlsl/NNERuntimeRDGGemm.cpp
@@ -4,8 +4,8 @@
 
 #include "NNEHlslShadersGemmCS.h"
 #include "NNEHlslShadersLog.h"
+#include "NNERuntimeRDGDataAttributeMap.h"
 #include "NNERuntimeRDGHlslHelper.h"
-#include "NNEAttributeMap.h"
 #include "NNETensor.h"
 #include "NNETypes.h"
 
@@ -57,7 +57,7 @@ namespace UE::NNERuntimeRDG::Private::Hlsl
 			return 0;
 		};
 		
-		virtual bool Initialize(TConstArrayView<NNE::FTensorDesc> InputTensorDescs, TConstArrayView<NNE::FTensorDesc> OutputTensorDescs, const NNE::FAttributeMap& Attributes) override
+		virtual bool Initialize(TConstArrayView<NNE::FTensorDesc> InputTensorDescs, TConstArrayView<NNE::FTensorDesc> OutputTensorDescs, const NNERuntimeRDGData::Internal::FAttributeMap& Attributes) override
 		{
 			check(InputTensorDescs.Num() >= 2 && InputTensorDescs.Num() <= 3);
 			check(OutputTensorDescs.Num() == 1);
@@ -162,17 +162,17 @@ namespace UE::NNERuntimeRDG::Private::Hlsl
 		}
 	};
 
-	bool ValidateGemmOperator(const NNE::FAttributeMap& AttributeMap, TConstArrayView<ENNETensorDataType> InputTypes, TConstArrayView<NNE::FSymbolicTensorShape> InputShapes)
+	bool ValidateGemmOperator(const NNERuntimeRDGData::Internal::FAttributeMap& AttributeMap, TConstArrayView<ENNETensorDataType> InputTypes, TConstArrayView<NNE::FSymbolicTensorShape> InputShapes)
 	{
 		bool bIsValid = true;
 
 		//This match version 9, 11, 13 of the Gemm operator see
 		//https://onnx.ai/onnx/operators/onnx__Gemm.html#l-onnx-doc-gemm
 		FAttributeValidator AttributeValidator;
-		AttributeValidator.AddOptional(TEXT("alpha"), ENNEAttributeDataType::Float);
-		AttributeValidator.AddOptional(TEXT("beta"), ENNEAttributeDataType::Float);
-		AttributeValidator.AddOptional(TEXT("transA"), ENNEAttributeDataType::Int32);
-		AttributeValidator.AddOptional(TEXT("transB"), ENNEAttributeDataType::Int32);
+		AttributeValidator.AddOptional(TEXT("alpha"), ENNERuntimeRDGDataAttributeDataType::Float);
+		AttributeValidator.AddOptional(TEXT("beta"), ENNERuntimeRDGDataAttributeDataType::Float);
+		AttributeValidator.AddOptional(TEXT("transA"), ENNERuntimeRDGDataAttributeDataType::Int32);
+		AttributeValidator.AddOptional(TEXT("transB"), ENNERuntimeRDGDataAttributeDataType::Int32);
 		bIsValid &= AttributeValidator.Validate(AttributeMap);
 
 		FInputValidator InputValidator;
diff --git a/Engine/Plugins/Experimental/NNERuntimeRDG/Source/NNERuntimeRDG/Private/Hlsl/NNERuntimeRDGGlobalPool.cpp b/Engine/Plugins/Experimental/NNERuntimeRDG/Source/NNERuntimeRDG/Private/Hlsl/NNERuntimeRDGGlobalPool.cpp
--- a/Engine/Plugins/Experimental/NNERuntimeRDG/Source/NNERuntimeRDG/Private/Hlsl/NNERuntimeRDGGlobalPool.cpp
+++ b/Engine/Plugins/Experimental/NNERuntimeRDG/Source/NNERuntimeRDG/Private/Hlsl/NNERuntimeRDGGlobalPool.cpp
@@ -52,7 +52,7 @@ namespace UE::NNERuntimeRDG::Private::Hlsl
 			return 0;
 		};
 
-		virtual bool Initialize(TConstArrayView<NNE::FTensorDesc> InputTensorDescs, TConstArrayView<NNE::FTensorDesc> OutputTensorDescs, const NNE::FAttributeMap& Attributes) override
+		virtual bool Initialize(TConstArrayView<NNE::FTensorDesc> InputTensorDescs, TConstArrayView<NNE::FTensorDesc> OutputTensorDescs, const NNERuntimeRDGData::Internal::FAttributeMap& Attributes) override
 		{
 			check(InputTensorDescs.Num() == 1);
 			check(OutputTensorDescs.Num() == 1);
@@ -98,7 +98,7 @@ namespace UE::NNERuntimeRDG::Private::Hlsl
 		}
 	};
 
-	bool ValidateGlobalPoolOperator(const NNE::FAttributeMap& AttributeMap, TConstArrayView<ENNETensorDataType> InputTypes, TConstArrayView<NNE::FSymbolicTensorShape> InputShapes)
+	bool ValidateGlobalPoolOperator(const NNERuntimeRDGData::Internal::FAttributeMap& AttributeMap, TConstArrayView<ENNETensorDataType> InputTypes, TConstArrayView<NNE::FSymbolicTensorShape> InputShapes)
 	{
 		bool bIsValid = true;
 
diff --git a/Engine/Plugins/Experimental/NNERuntimeRDG/Source/NNERuntimeRDG/Private/Hlsl/NNERuntimeRDGIdentity.cpp b/Engine/Plugins/Experimental/NNERuntimeRDG/Source/NNERuntimeRDG/Private/Hlsl/NNERuntimeRDGIdentity.cpp
--- a/Engine/Plugins/Experimental/NNERuntimeRDG/Source/NNERuntimeRDG/Private/Hlsl/NNERuntimeRDGIdentity.cpp
+++ b/Engine/Plugins/Experimental/NNERuntimeRDG/Source/NNERuntimeRDG/Private/Hlsl/NNERuntimeRDGIdentity.cpp
@@ -37,7 +37,7 @@ namespace UE::NNERuntimeRDG::Private::Hlsl
 			return 0;
 		};
 
-		virtual bool Initialize(TConstArrayView<NNE::FTensorDesc> InputTensorDescs, TConstArrayView<NNE::FTensorDesc> OutputTensorDescs, const NNE::FAttributeMap& Attributes) override
+		virtual bool Initialize(TConstArrayView<NNE::FTensorDesc> InputTensorDescs, TConstArrayView<NNE::FTensorDesc> OutputTensorDescs, const NNERuntimeRDGData::Internal::FAttributeMap& Attributes) override
 		{
 			check(InputTensorDescs.Num() == 1);
 			check(OutputTensorDescs.Num() == 1);
@@ -62,7 +62,7 @@ namespace UE::NNERuntimeRDG::Private::Hlsl
 		}
 	};
 
-	bool ValidateIdentityOperator(const NNE::FAttributeMap& AttributeMap, TConstArrayView<ENNETensorDataType> InputTypes, TConstArrayView<NNE::FSymbolicTensorShape> InputIdentitys)
+	bool ValidateIdentityOperator(const NNERuntimeRDGData::Internal::FAttributeMap& AttributeMap, TConstArrayView<ENNETensorDataType> InputTypes, TConstArrayView<NNE::FSymbolicTensorShape> InputIdentitys)
 	{
 		bool bIsValid = true;
 
diff --git a/Engine/Plugins/Experimental/NNERuntimeRDG/Source/NNERuntimeRDG/Private/Hlsl/NNERuntimeRDGInstanceNormalization.cpp b/Engine/Plugins/Experimental/NNERuntimeRDG/Source/NNERuntimeRDG/Private/Hlsl/NNERuntimeRDGInstanceNormalization.cpp
--- a/Engine/Plugins/Experimental/NNERuntimeRDG/Source/NNERuntimeRDG/Private/Hlsl/NNERuntimeRDGInstanceNormalization.cpp
+++ b/Engine/Plugins/Experimental/NNERuntimeRDG/Source/NNERuntimeRDG/Private/Hlsl/NNERuntimeRDGInstanceNormalization.cpp
@@ -77,7 +77,7 @@ namespace UE::NNERuntimeRDG::Private::Hlsl
 			return 0;
 		};
 
-		virtual bool Initialize(TConstArrayView<NNE::FTensorDesc> InputTensorDescs, TConstArrayView<NNE::FTensorDesc> OutputTensorDescs, const NNE::FAttributeMap& Attributes) override
+		virtual bool Initialize(TConstArrayView<NNE::FTensorDesc> InputTensorDescs, TConstArrayView<NNE::FTensorDesc> OutputTensorDescs, const NNERuntimeRDGData::Internal::FAttributeMap& Attributes) override
 		{
 			using namespace UE::NNEHlslShaders::Internal;
 
@@ -166,12 +166,12 @@ namespace UE::NNERuntimeRDG::Private::Hlsl
 		}
 	};
 
-	bool ValidateInstanceNormalizationOperator(const NNE::FAttributeMap& AttributeMap, TConstArrayView<ENNETensorDataType> InputTypes, TConstArrayView<NNE::FSymbolicTensorShape> InputShapes)
+	bool ValidateInstanceNormalizationOperator(const NNERuntimeRDGData::Internal::FAttributeMap& AttributeMap, TConstArrayView<ENNETensorDataType> InputTypes, TConstArrayView<NNE::FSymbolicTensorShape> InputShapes)
 	{
 		bool bIsValid = true;
 
 		FAttributeValidator AttributeValidator;
-		AttributeValidator.AddOptional(TEXT("epsilon"), ENNEAttributeDataType::Float);
+		AttributeValidator.AddOptional(TEXT("epsilon"), ENNERuntimeRDGDataAttributeDataType::Float);
 		bIsValid &= AttributeValidator.Validate(AttributeMap);
 
 		FInputValidator InputValidator;
diff --git a/Engine/Plugins/Experimental/NNERuntimeRDG/Source/NNERuntimeRDG/Private/Hlsl/NNERuntimeRDGLayerNormalization.cpp b/Engine/Plugins/Experimental/NNERuntimeRDG/Source/NNERuntimeRDG/Private/Hlsl/NNERuntimeRDGLayerNormalization.cpp
--- a/Engine/Plugins/Experimental/NNERuntimeRDG/Source/NNERuntimeRDG/Private/Hlsl/NNERuntimeRDGLayerNormalization.cpp
+++ b/Engine/Plugins/Experimental/NNERuntimeRDG/Source/NNERuntimeRDG/Private/Hlsl/NNERuntimeRDGLayerNormalization.cpp
@@ -98,7 +98,7 @@ namespace UE::NNERuntimeRDG::Private::Hlsl
 			return 0;
 		};
 
-		virtual bool Initialize(TConstArrayView<NNE::FTensorDesc> InputTensorDescs, TConstArrayView<NNE::FTensorDesc> OutputTensorDescs, const NNE::FAttributeMap& Attributes) override
+		virtual bool Initialize(TConstArrayView<NNE::FTensorDesc> InputTensorDescs, TConstArrayView<NNE::FTensorDesc> OutputTensorDescs, const NNERuntimeRDGData::Internal::FAttributeMap& Attributes) override
 		{
 			check(InputTensorDescs.Num() >= 2 && InputTensorDescs.Num() <= 3);
 			check(OutputTensorDescs.Num() >= 1 && OutputTensorDescs.Num() <= 3);
@@ -248,14 +248,14 @@ namespace UE::NNERuntimeRDG::Private::Hlsl
 		}
 	};
 
-	bool ValidateLayerNormalizationOperator(const NNE::FAttributeMap& AttributeMap, TConstArrayView<ENNETensorDataType> InputTypes, TConstArrayView<NNE::FSymbolicTensorShape> InputShapes)
+	bool ValidateLayerNormalizationOperator(const NNERuntimeRDGData::Internal::FAttributeMap& AttributeMap, TConstArrayView<ENNETensorDataType> InputTypes, TConstArrayView<NNE::FSymbolicTensorShape> InputShapes)
 	{
 		bool bIsValid = true;
 
 		FAttributeValidator AttributeValidator;
-		AttributeValidator.AddOptional(TEXT("axis"), ENNEAttributeDataType::Int32);
-		AttributeValidator.AddOptional(TEXT("epsilon"), ENNEAttributeDataType::Float);
-		AttributeValidator.AddOptional(TEXT("stash_type"), ENNEAttributeDataType::Int32);
+		AttributeValidator.AddOptional(TEXT("axis"), ENNERuntimeRDGDataAttributeDataType::Int32);
+		AttributeValidator.AddOptional(TEXT("epsilon"), ENNERuntimeRDGDataAttributeDataType::Float);
+		AttributeValidator.AddOptional(TEXT("stash_type"), ENNERuntimeRDGDataAttributeDataType::Int32);
 		bIsValid &= AttributeValidator.Validate(AttributeMap);
 
 		FInputValidator InputValidator;
diff --git a/Engine/Plugins/Experimental/NNERuntimeRDG/Source/NNERuntimeRDG/Private/Hlsl/NNERuntimeRDGMatMul.cpp b/Engine/Plugins/Experimental/NNERuntimeRDG/Source/NNERuntimeRDG/Private/Hlsl/NNERuntimeRDGMatMul.cpp
--- a/Engine/Plugins/Experimental/NNERuntimeRDG/Source/NNERuntimeRDG/Private/Hlsl/NNERuntimeRDGMatMul.cpp
+++ b/Engine/Plugins/Experimental/NNERuntimeRDG/Source/NNERuntimeRDG/Private/Hlsl/NNERuntimeRDGMatMul.cpp
@@ -74,7 +74,7 @@ namespace UE::NNERuntimeRDG::Private::Hlsl
 			return 0;
 		};
 
-		virtual bool Initialize(TConstArrayView<NNE::FTensorDesc> InputTensorDescs, TConstArrayView<NNE::FTensorDesc> OutputTensorDescs, const NNE::FAttributeMap& Attributes) override
+		virtual bool Initialize(TConstArrayView<NNE::FTensorDesc> InputTensorDescs, TConstArrayView<NNE::FTensorDesc> OutputTensorDescs, const NNERuntimeRDGData::Internal::FAttributeMap& Attributes) override
 		{
 			check(InputTensorDescs.Num() == 2);
 			check(OutputTensorDescs.Num() == 1);
@@ -141,7 +141,7 @@ namespace UE::NNERuntimeRDG::Private::Hlsl
 		}
 	};
 
-	bool ValidateMatMulOperator(const NNE::FAttributeMap& AttributeMap, TConstArrayView<ENNETensorDataType> InputTypes, TConstArrayView<NNE::FSymbolicTensorShape> InputShapes)
+	bool ValidateMatMulOperator(const NNERuntimeRDGData::Internal::FAttributeMap& AttributeMap, TConstArrayView<ENNETensorDataType> InputTypes, TConstArrayView<NNE::FSymbolicTensorShape> InputShapes)
 	{
 		bool bIsValid = true;
 
diff --git a/Engine/Plugins/Experimental/NNERuntimeRDG/Source/NNERuntimeRDG/Private/Hlsl/NNERuntimeRDGPad.cpp b/Engine/Plugins/Experimental/NNERuntimeRDG/Source/NNERuntimeRDG/Private/Hlsl/NNERuntimeRDGPad.cpp
--- a/Engine/Plugins/Experimental/NNERuntimeRDG/Source/NNERuntimeRDG/Private/Hlsl/NNERuntimeRDGPad.cpp
+++ b/Engine/Plugins/Experimental/NNERuntimeRDG/Source/NNERuntimeRDG/Private/Hlsl/NNERuntimeRDGPad.cpp
@@ -154,7 +154,7 @@ namespace UE::NNERuntimeRDG::Private::Hlsl
 			return 0;
 		};
 
-		virtual bool Initialize(TConstArrayView<NNE::FTensorDesc> InputTensorDescs, TConstArrayView<NNE::FTensorDesc> OutputTensorDescs, const NNE::FAttributeMap& Attributes) override
+		virtual bool Initialize(TConstArrayView<NNE::FTensorDesc> InputTensorDescs, TConstArrayView<NNE::FTensorDesc> OutputTensorDescs, const NNERuntimeRDGData::Internal::FAttributeMap& Attributes) override
 		{
 			using namespace UE::NNEHlslShaders::Internal;
 			
@@ -240,16 +240,16 @@ namespace UE::NNERuntimeRDG::Private::Hlsl
 	};
 
 	template<int Version>
-	bool ValidatePadOperator(const NNE::FAttributeMap& AttributeMap, TConstArrayView<ENNETensorDataType> InputTypes, TConstArrayView<NNE::FSymbolicTensorShape> InputShapes)
+	bool ValidatePadOperator(const NNERuntimeRDGData::Internal::FAttributeMap& AttributeMap, TConstArrayView<ENNETensorDataType> InputTypes, TConstArrayView<NNE::FSymbolicTensorShape> InputShapes)
 	{
 		bool bIsValid = true;
 
 		FAttributeValidator AttributeValidator;
-		AttributeValidator.AddOptional(TEXT("mode"), ENNEAttributeDataType::String);
+		AttributeValidator.AddOptional(TEXT("mode"), ENNERuntimeRDGDataAttributeDataType::String);
 		if constexpr (Version < 11)
 		{
-			AttributeValidator.AddRequired(TEXT("pads"), ENNEAttributeDataType::Int32Array);
-			AttributeValidator.AddOptional(TEXT("value"), ENNEAttributeDataType::Float);
+			AttributeValidator.AddRequired(TEXT("pads"), ENNERuntimeRDGDataAttributeDataType::Int32Array);
+			AttributeValidator.AddOptional(TEXT("value"), ENNERuntimeRDGDataAttributeDataType::Float);
 		}
 		bIsValid &= AttributeValidator.Validate(AttributeMap);
 		
diff --git a/Engine/Plugins/Experimental/NNERuntimeRDG/Source/NNERuntimeRDG/Private/Hlsl/NNERuntimeRDGPool.cpp b/Engine/Plugins/Experimental/NNERuntimeRDG/Source/NNERuntimeRDG/Private/Hlsl/NNERuntimeRDGPool.cpp
--- a/Engine/Plugins/Experimental/NNERuntimeRDG/Source/NNERuntimeRDG/Private/Hlsl/NNERuntimeRDGPool.cpp
+++ b/Engine/Plugins/Experimental/NNERuntimeRDG/Source/NNERuntimeRDG/Private/Hlsl/NNERuntimeRDGPool.cpp
@@ -96,7 +96,7 @@ namespace UE::NNERuntimeRDG::Private::Hlsl
 			return 0;
 		};
 
-		virtual bool Initialize(TConstArrayView<NNE::FTensorDesc> InputTensorDescs, TConstArrayView<NNE::FTensorDesc> OutputTensorDescs, const NNE::FAttributeMap& Attributes) override
+		virtual bool Initialize(TConstArrayView<NNE::FTensorDesc> InputTensorDescs, TConstArrayView<NNE::FTensorDesc> OutputTensorDescs, const NNERuntimeRDGData::Internal::FAttributeMap& Attributes) override
 		{
 			check(InputTensorDescs.Num() == 1);
 
@@ -268,26 +268,26 @@ namespace UE::NNERuntimeRDG::Private::Hlsl
 	};
 
 	template< UE::NNEHlslShaders::Internal::EPoolOperatorType PoolOperatorType >
-	bool ValidatePoolOperator(const NNE::FAttributeMap& AttributeMap, TConstArrayView<ENNETensorDataType> InputTypes, TConstArrayView<NNE::FSymbolicTensorShape> InputPools)
+	bool ValidatePoolOperator(const NNERuntimeRDGData::Internal::FAttributeMap& AttributeMap, TConstArrayView<ENNETensorDataType> InputTypes, TConstArrayView<NNE::FSymbolicTensorShape> InputPools)
 	{
 		bool bIsValid = true;
 
 		FAttributeValidator AttributeValidator;
-		AttributeValidator.AddOptional(TEXT("auto_pad"), ENNEAttributeDataType::String);
-		AttributeValidator.AddRequired(TEXT("kernel_shape"), ENNEAttributeDataType::Int32Array);
-		AttributeValidator.AddOptional(TEXT("pads"), ENNEAttributeDataType::Int32Array);
-		AttributeValidator.AddOptional(TEXT("ceil_mode"), ENNEAttributeDataType::Int32);
+		AttributeValidator.AddOptional(TEXT("auto_pad"), ENNERuntimeRDGDataAttributeDataType::String);
+		AttributeValidator.AddRequired(TEXT("kernel_shape"), ENNERuntimeRDGDataAttributeDataType::Int32Array);
+		AttributeValidator.AddOptional(TEXT("pads"), ENNERuntimeRDGDataAttributeDataType::Int32Array);
+		AttributeValidator.AddOptional(TEXT("ceil_mode"), ENNERuntimeRDGDataAttributeDataType::Int32);
 		if constexpr (PoolOperatorType == UE::NNEHlslShaders::Internal::EPoolOperatorType::MAX_POOL)
 		{
-			AttributeValidator.AddOptional(TEXT("dilations"), ENNEAttributeDataType::Int32Array);
-			AttributeValidator.AddOptional(TEXT("storage_order"), ENNEAttributeDataType::Int32);//Unused, only needed for 2nd output itself not supported, see [REDACTED_PROFILE_URL]/onnx/issues/1370
+			AttributeValidator.AddOptional(TEXT("dilations"), ENNERuntimeRDGDataAttributeDataType::Int32Array);
+			AttributeValidator.AddOptional(TEXT("storage_order"), ENNERuntimeRDGDataAttributeDataType::Int32);//Unused, only needed for 2nd output itself not supported, see [REDACTED_PROFILE_URL]/onnx/issues/1370
 		}
 		else
 		{
-			AttributeValidator.AddOptional(TEXT("count_include_pad"), ENNEAttributeDataType::Int32);
+			AttributeValidator.AddOptional(TEXT("count_include_pad"), ENNERuntimeRDGDataAttributeDataType::Int32);
 		}
 
-		AttributeValidator.AddOptional(TEXT("strides"), ENNEAttributeDataType::Int32Array);
+		AttributeValidator.AddOptional(TEXT("strides"), ENNERuntimeRDGDataAttributeDataType::Int32Array);
 		bIsValid &= AttributeValidator.Validate(AttributeMap);
 
 		FInputValidator InputValidator;
diff --git a/Engine/Plugins/Experimental/NNERuntimeRDG/Source/NNERuntimeRDG/Private/Hlsl/NNERuntimeRDGReduce.cpp b/Engine/Plugins/Experimental/NNERuntimeRDG/Source/NNERuntimeRDG/Private/Hlsl/NNERuntimeRDGReduce.cpp
--- a/Engine/Plugins/Experimental/NNERuntimeRDG/Source/NNERuntimeRDG/Private/Hlsl/NNERuntimeRDGReduce.cpp
+++ b/Engine/Plugins/Experimental/NNERuntimeRDG/Source/NNERuntimeRDG/Private/Hlsl/NNERuntimeRDGReduce.cpp
@@ -108,7 +108,7 @@ namespace UE::NNERuntimeRDG::Private::Hlsl
 			return 0;
 		};
 
-		virtual bool Initialize(TConstArrayView<NNE::FTensorDesc> InputTensorDescs, TConstArrayView<NNE::FTensorDesc> OutputTensorDescs, const NNE::FAttributeMap& Attributes) override
+		virtual bool Initialize(TConstArrayView<NNE::FTensorDesc> InputTensorDescs, TConstArrayView<NNE::FTensorDesc> OutputTensorDescs, const NNERuntimeRDGData::Internal::FAttributeMap& Attributes) override
 		{
 			check(bAxesAsInput ? (InputTensorDescs.Num() >= 1 && InputTensorDescs.Num() <= 2) : InputTensorDescs.Num() == 1);
 			check(OutputTensorDescs.Num() == 1);
@@ -193,7 +193,7 @@ namespace UE::NNERuntimeRDG::Private::Hlsl
 	};
 
 	template< UE::NNEHlslShaders::Internal::EReduceOperatorType ReduceOperatorType, int Version>
-	bool ValidateReduceOperator(const NNE::FAttributeMap& AttributeMap, TConstArrayView<ENNETensorDataType> InputTypes, TConstArrayView<NNE::FSymbolicTensorShape> InputShapes)
+	bool ValidateReduceOperator(const NNERuntimeRDGData::Internal::FAttributeMap& AttributeMap, TConstArrayView<ENNETensorDataType> InputTypes, TConstArrayView<NNE::FSymbolicTensorShape> InputShapes)
 	{
 		bool bIsValid = true;
 
@@ -202,14 +202,14 @@ namespace UE::NNERuntimeRDG::Private::Hlsl
 		// ReduceMean-13 is also supported
 		// [REDACTED_PROFILE_URL]/onnx/blob/main/docs/Changelog.md#Reduce-1
 		FAttributeValidator AttributeValidator;
-		AttributeValidator.AddOptional(TEXT("keepdims"), ENNEAttributeDataType::Int32);
+		AttributeValidator.AddOptional(TEXT("keepdims"), ENNERuntimeRDGDataAttributeDataType::Int32);
 		if constexpr (!FReduceOperator<ReduceOperatorType, Version>::bAxesAsInput)
 		{
-			AttributeValidator.AddOptional(TEXT("axes"), ENNEAttributeDataType::Int32Array);
+			AttributeValidator.AddOptional(TEXT("axes"), ENNERuntimeRDGDataAttributeDataType::Int32Array);
 		}
 		else
 		{
-			AttributeValidator.AddOptional(TEXT("noop_with_empty_axes"), ENNEAttributeDataType::Int32);
+			AttributeValidator.AddOptional(TEXT("noop_with_empty_axes"), ENNERuntimeRDGDataAttributeDataType::Int32);
 		}
 		bIsValid &= AttributeValidator.Validate(AttributeMap);
 
diff --git a/Engine/Plugins/Experimental/NNERuntimeRDG/Source/NNERuntimeRDG/Private/Hlsl/NNERuntimeRDGReshape.cpp b/Engine/Plugins/Experimental/NNERuntimeRDG/Source/NNERuntimeRDG/Private/Hlsl/NNERuntimeRDGReshape.cpp
--- a/Engine/Plugins/Experimental/NNERuntimeRDG/Source/NNERuntimeRDG/Private/Hlsl/NNERuntimeRDGReshape.cpp
+++ b/Engine/Plugins/Experimental/NNERuntimeRDG/Source/NNERuntimeRDG/Private/Hlsl/NNERuntimeRDGReshape.cpp
@@ -59,7 +59,7 @@ namespace UE::NNERuntimeRDG::Private::Hlsl
 			return 0;
 		};
 
-		virtual bool Initialize(TConstArrayView<NNE::FTensorDesc> InputTensorDescs, TConstArrayView<NNE::FTensorDesc> OutputTensorDescs, const NNE::FAttributeMap& Attributes) override
+		virtual bool Initialize(TConstArrayView<NNE::FTensorDesc> InputTensorDescs, TConstArrayView<NNE::FTensorDesc> OutputTensorDescs, const NNERuntimeRDGData::Internal::FAttributeMap& Attributes) override
 		{
 			check(InputTensorDescs.Num() == 2);
 			check(OutputTensorDescs.Num() == 1);
@@ -86,12 +86,12 @@ namespace UE::NNERuntimeRDG::Private::Hlsl
 		}
 	};
 
-	bool ValidateReshapeOperator(const NNE::FAttributeMap& AttributeMap, TConstArrayView<ENNETensorDataType> InputTypes, TConstArrayView<NNE::FSymbolicTensorShape> InputReshapes)
+	bool ValidateReshapeOperator(const NNERuntimeRDGData::Internal::FAttributeMap& AttributeMap, TConstArrayView<ENNETensorDataType> InputTypes, TConstArrayView<NNE::FSymbolicTensorShape> InputReshapes)
 	{
 		bool bIsValid = true;
 
 		FAttributeValidator AttributeValidator;
-		AttributeValidator.AddOptional(TEXT("allowzero"), ENNEAttributeDataType::Int32);
+		AttributeValidator.AddOptional(TEXT("allowzero"), ENNERuntimeRDGDataAttributeDataType::Int32);
 		bIsValid &= AttributeValidator.Validate(AttributeMap);
 
 		FInputValidator InputValidator;
diff --git a/Engine/Plugins/Experimental/NNERuntimeRDG/Source/NNERuntimeRDG/Private/Hlsl/NNERuntimeRDGResize.cpp b/Engine/Plugins/Experimental/NNERuntimeRDG/Source/NNERuntimeRDG/Private/Hlsl/NNERuntimeRDGResize.cpp
--- a/Engine/Plugins/Experimental/NNERuntimeRDG/Source/NNERuntimeRDG/Private/Hlsl/NNERuntimeRDGResize.cpp
+++ b/Engine/Plugins/Experimental/NNERuntimeRDG/Source/NNERuntimeRDG/Private/Hlsl/NNERuntimeRDGResize.cpp
@@ -170,7 +170,7 @@ namespace UE::NNERuntimeRDG::Private::Hlsl
 			return 0;
 		};
 
-		virtual bool Initialize(TConstArrayView<NNE::FTensorDesc> InputTensorDescs, TConstArrayView<NNE::FTensorDesc> OutputTensorDescs, const NNE::FAttributeMap& Attributes) override
+		virtual bool Initialize(TConstArrayView<NNE::FTensorDesc> InputTensorDescs, TConstArrayView<NNE::FTensorDesc> OutputTensorDescs, const NNERuntimeRDGData::Internal::FAttributeMap& Attributes) override
 		{
 			check(InputTensorDescs.Num() >= 1 && InputTensorDescs.Num() <= 4);
 			check(OutputTensorDescs.Num() == 1);
@@ -308,22 +308,22 @@ namespace UE::NNERuntimeRDG::Private::Hlsl
 		}
 	};
 
-	bool ValidateResizeOperator(const NNE::FAttributeMap& AttributeMap, TConstArrayView<ENNETensorDataType> InputTypes, TConstArrayView<NNE::FSymbolicTensorShape> InputShapes)
+	bool ValidateResizeOperator(const NNERuntimeRDGData::Internal::FAttributeMap& AttributeMap, TConstArrayView<ENNETensorDataType> InputTypes, TConstArrayView<NNE::FSymbolicTensorShape> InputShapes)
 	{
 		using namespace UE::NNEHlslShaders::Internal;
 		
 		bool bIsValid = true;
 
 		FAttributeValidator AttributeValidator;
-		AttributeValidator.AddOptional(TEXT("antialias"), ENNEAttributeDataType::Int32);
-		AttributeValidator.AddOptional(TEXT("axes"), ENNEAttributeDataType::Int32Array);
-		AttributeValidator.AddOptional(TEXT("coordinate_transformation_mode"), ENNEAttributeDataType::String);
-		AttributeValidator.AddOptional(TEXT("cubic_coeff_a"), ENNEAttributeDataType::Float);
-		AttributeValidator.AddOptional(TEXT("exclude_outside"), ENNEAttributeDataType::Int32);
-		AttributeValidator.AddOptional(TEXT("extrapolation_value"), ENNEAttributeDataType::Float);
-		AttributeValidator.AddOptional(TEXT("keep_aspect_ratio_policy"), ENNEAttributeDataType::String);
-		AttributeValidator.AddOptional(TEXT("mode"), ENNEAttributeDataType::String);
-		AttributeValidator.AddOptional(TEXT("nearest_mode"), ENNEAttributeDataType::String);
+		AttributeValidator.AddOptional(TEXT("antialias"), ENNERuntimeRDGDataAttributeDataType::Int32);
+		AttributeValidator.AddOptional(TEXT("axes"), ENNERuntimeRDGDataAttributeDataType::Int32Array);
+		AttributeValidator.AddOptional(TEXT("coordinate_transformation_mode"), ENNERuntimeRDGDataAttributeDataType::String);
+		AttributeValidator.AddOptional(TEXT("cubic_coeff_a"), ENNERuntimeRDGDataAttributeDataType::Float);
+		AttributeValidator.AddOptional(TEXT("exclude_outside"), ENNERuntimeRDGDataAttributeDataType::Int32);
+		AttributeValidator.AddOptional(TEXT("extrapolation_value"), ENNERuntimeRDGDataAttributeDataType::Float);
+		AttributeValidator.AddOptional(TEXT("keep_aspect_ratio_policy"), ENNERuntimeRDGDataAttributeDataType::String);
+		AttributeValidator.AddOptional(TEXT("mode"), ENNERuntimeRDGDataAttributeDataType::String);
+		AttributeValidator.AddOptional(TEXT("nearest_mode"), ENNERuntimeRDGDataAttributeDataType::String);
 		bIsValid &= AttributeValidator.Validate(AttributeMap);
 
 		FInputValidator InputValidator;
diff --git a/Engine/Plugins/Experimental/NNERuntimeRDG/Source/NNERuntimeRDG/Private/Hlsl/NNERuntimeRDGScatterND.cpp b/Engine/Plugins/Experimental/NNERuntimeRDG/Source/NNERuntimeRDG/Private/Hlsl/NNERuntimeRDGScatterND.cpp
--- a/Engine/Plugins/Experimental/NNERuntimeRDG/Source/NNERuntimeRDG/Private/Hlsl/NNERuntimeRDGScatterND.cpp
+++ b/Engine/Plugins/Experimental/NNERuntimeRDG/Source/NNERuntimeRDG/Private/Hlsl/NNERuntimeRDGScatterND.cpp
@@ -73,7 +73,7 @@ namespace UE::NNERuntimeRDG::Private::Hlsl
 			return 0;
 		};
 
-		virtual bool Initialize(TConstArrayView<NNE::FTensorDesc> InputTensorDescs, TConstArrayView<NNE::FTensorDesc> OutputTensorDescs, const NNE::FAttributeMap& Attributes) override
+		virtual bool Initialize(TConstArrayView<NNE::FTensorDesc> InputTensorDescs, TConstArrayView<NNE::FTensorDesc> OutputTensorDescs, const NNERuntimeRDGData::Internal::FAttributeMap& Attributes) override
 		{
             using namespace UE::NNEHlslShaders::Internal;
 			check(InputTensorDescs.Num() == 3);
@@ -159,12 +159,12 @@ namespace UE::NNERuntimeRDG::Private::Hlsl
 		}
 	};
 
-	bool ValidateScatterNDOperator(const NNE::FAttributeMap& AttributeMap, TConstArrayView<ENNETensorDataType> InputTypes, TConstArrayView<NNE::FSymbolicTensorShape> InputScatterNDs)
+	bool ValidateScatterNDOperator(const NNERuntimeRDGData::Internal::FAttributeMap& AttributeMap, TConstArrayView<ENNETensorDataType> InputTypes, TConstArrayView<NNE::FSymbolicTensorShape> InputScatterNDs)
 	{
 		bool bIsValid = true;
 
 		FAttributeValidator AttributeValidator;
-		AttributeValidator.AddOptional(TEXT("reduction"), ENNEAttributeDataType::String);
+		AttributeValidator.AddOptional(TEXT("reduction"), ENNERuntimeRDGDataAttributeDataType::String);
 		bIsValid &= AttributeValidator.Validate(AttributeMap);
 
 		FInputValidator InputValidator;
diff --git a/Engine/Plugins/Experimental/NNERuntimeRDG/Source/NNERuntimeRDG/Private/Hlsl/NNERuntimeRDGShape.cpp b/Engine/Plugins/Experimental/NNERuntimeRDG/Source/NNERuntimeRDG/Private/Hlsl/NNERuntimeRDGShape.cpp
--- a/Engine/Plugins/Experimental/NNERuntimeRDG/Source/NNERuntimeRDG/Private/Hlsl/NNERuntimeRDGShape.cpp
+++ b/Engine/Plugins/Experimental/NNERuntimeRDG/Source/NNERuntimeRDG/Private/Hlsl/NNERuntimeRDGShape.cpp
@@ -45,7 +45,7 @@ namespace UE::NNERuntimeRDG::Private::Hlsl
 			return 0;
 		};
 
-		virtual bool Initialize(TConstArrayView<NNE::FTensorDesc> InputTensorDescs, TConstArrayView<NNE::FTensorDesc> OutputTensorDescs, const NNE::FAttributeMap& Attributes) override
+		virtual bool Initialize(TConstArrayView<NNE::FTensorDesc> InputTensorDescs, TConstArrayView<NNE::FTensorDesc> OutputTensorDescs, const NNERuntimeRDGData::Internal::FAttributeMap& Attributes) override
 		{
 			check(InputTensorDescs.Num() == 1);
 			check(OutputTensorDescs.Num() == 1);
@@ -65,7 +65,7 @@ namespace UE::NNERuntimeRDG::Private::Hlsl
 		}
 	};
 
-	bool ValidateShapeOperator(const NNE::FAttributeMap& AttributeMap, TConstArrayView<ENNETensorDataType> InputTypes, TConstArrayView<NNE::FSymbolicTensorShape> InputShapes)
+	bool ValidateShapeOperator(const NNERuntimeRDGData::Internal::FAttributeMap& AttributeMap, TConstArrayView<ENNETensorDataType> InputTypes, TConstArrayView<NNE::FSymbolicTensorShape> InputShapes)
 	{
 		bool bIsValid = true;
 
diff --git a/Engine/Plugins/Experimental/NNERuntimeRDG/Source/NNERuntimeRDG/Private/Hlsl/NNERuntimeRDGSize.cpp b/Engine/Plugins/Experimental/NNERuntimeRDG/Source/NNERuntimeRDG/Private/Hlsl/NNERuntimeRDGSize.cpp
--- a/Engine/Plugins/Experimental/NNERuntimeRDG/Source/NNERuntimeRDG/Private/Hlsl/NNERuntimeRDGSize.cpp
+++ b/Engine/Plugins/Experimental/NNERuntimeRDG/Source/NNERuntimeRDG/Private/Hlsl/NNERuntimeRDGSize.cpp
@@ -38,7 +38,7 @@ namespace UE::NNERuntimeRDG::Private::Hlsl
 			return 0;
 		};
 
-		virtual bool Initialize(TConstArrayView<NNE::FTensorDesc> InputTensorDescs, TConstArrayView<NNE::FTensorDesc> OutputTensorDescs, const NNE::FAttributeMap& Attributes) override
+		virtual bool Initialize(TConstArrayView<NNE::FTensorDesc> InputTensorDescs, TConstArrayView<NNE::FTensorDesc> OutputTensorDescs, const NNERuntimeRDGData::Internal::FAttributeMap& Attributes) override
 		{
 			check(InputTensorDescs.Num() == 1);
 			check(OutputTensorDescs.Num() == 1);
@@ -53,7 +53,7 @@ namespace UE::NNERuntimeRDG::Private::Hlsl
 		}
 	};
 
-	bool ValidateSizeOperator(const NNE::FAttributeMap& AttributeMap, TConstArrayView<ENNETensorDataType> InputTypes, TConstArrayView<NNE::FSymbolicTensorShape> InputSizes)
+	bool ValidateSizeOperator(const NNERuntimeRDGData::Internal::FAttributeMap& AttributeMap, TConstArrayView<ENNETensorDataType> InputTypes, TConstArrayView<NNE::FSymbolicTensorShape> InputSizes)
 	{
 		bool bIsValid = true;
 
diff --git a/Engine/Plugins/Experimental/NNERuntimeRDG/Source/NNERuntimeRDG/Private/Hlsl/NNERuntimeRDGSlice.cpp b/Engine/Plugins/Experimental/NNERuntimeRDG/Source/NNERuntimeRDG/Private/Hlsl/NNERuntimeRDGSlice.cpp
--- a/Engine/Plugins/Experimental/NNERuntimeRDG/Source/NNERuntimeRDG/Private/Hlsl/NNERuntimeRDGSlice.cpp
+++ b/Engine/Plugins/Experimental/NNERuntimeRDG/Source/NNERuntimeRDG/Private/Hlsl/NNERuntimeRDGSlice.cpp
@@ -232,7 +232,7 @@ namespace UE::NNERuntimeRDG::Private::Hlsl
 			return 0;
 		};
 
-		virtual bool Initialize(TConstArrayView<NNE::FTensorDesc> InputTensorDescs, TConstArrayView<NNE::FTensorDesc> OutputTensorDescs, const NNE::FAttributeMap& Attributes) override
+		virtual bool Initialize(TConstArrayView<NNE::FTensorDesc> InputTensorDescs, TConstArrayView<NNE::FTensorDesc> OutputTensorDescs, const NNERuntimeRDGData::Internal::FAttributeMap& Attributes) override
 		{
 			check(OutputTensorDescs.Num() == 1);
 
@@ -318,14 +318,14 @@ namespace UE::NNERuntimeRDG::Private::Hlsl
 		}
 	};
 
-	bool ValidateSliceOperatorOpset1(const NNE::FAttributeMap& AttributeMap, TConstArrayView<ENNETensorDataType> InputTypes, TConstArrayView<NNE::FSymbolicTensorShape> InputSlices)
+	bool ValidateSliceOperatorOpset1(const NNERuntimeRDGData::Internal::FAttributeMap& AttributeMap, TConstArrayView<ENNETensorDataType> InputTypes, TConstArrayView<NNE::FSymbolicTensorShape> InputSlices)
 	{
 		bool bIsValid = true;
 
 		FAttributeValidator AttributeValidator;
-		AttributeValidator.AddOptional(TEXT("axes"), ENNEAttributeDataType::Int32Array);
-		AttributeValidator.AddRequired(TEXT("ends"), ENNEAttributeDataType::Int32Array);
-		AttributeValidator.AddRequired(TEXT("starts"), ENNEAttributeDataType::Int32Array);
+		AttributeValidator.AddOptional(TEXT("axes"), ENNERuntimeRDGDataAttributeDataType::Int32Array);
+		AttributeValidator.AddRequired(TEXT("ends"), ENNERuntimeRDGDataAttributeDataType::Int32Array);
+		AttributeValidator.AddRequired(TEXT("starts"), ENNERuntimeRDGDataAttributeDataType::Int32Array);
 		bIsValid &= AttributeValidator.Validate(AttributeMap);
 
 		FInputValidator InputValidator;
@@ -339,7 +339,7 @@ namespace UE::NNERuntimeRDG::Private::Hlsl
 		return bIsValid;
 	}
 
-	bool ValidateSliceOperatorOpset10To13(const NNE::FAttributeMap& AttributeMap, TConstArrayView<ENNETensorDataType> InputTypes, TConstArrayView<NNE::FSymbolicTensorShape> InputSlices)
+	bool ValidateSliceOperatorOpset10To13(const NNERuntimeRDGData::Internal::FAttributeMap& AttributeMap, TConstArrayView<ENNETensorDataType> InputTypes, TConstArrayView<NNE::FSymbolicTensorShape> InputSlices)
 	{
 		bool bIsValid = true;
 
diff --git a/Engine/Plugins/Experimental/NNERuntimeRDG/Source/NNERuntimeRDG/Private/Hlsl/NNERuntimeRDGSoftmax.cpp b/Engine/Plugins/Experimental/NNERuntimeRDG/Source/NNERuntimeRDG/Private/Hlsl/NNERuntimeRDGSoftmax.cpp
--- a/Engine/Plugins/Experimental/NNERuntimeRDG/Source/NNERuntimeRDG/Private/Hlsl/NNERuntimeRDGSoftmax.cpp
+++ b/Engine/Plugins/Experimental/NNERuntimeRDG/Source/NNERuntimeRDG/Private/Hlsl/NNERuntimeRDGSoftmax.cpp
@@ -44,7 +44,7 @@ namespace UE::NNERuntimeRDG::Private::Hlsl
 			return 0;
 		};
 
-		virtual bool Initialize(TConstArrayView<NNE::FTensorDesc> InputTensorDescs, TConstArrayView<NNE::FTensorDesc> OutputTensorDescs, const NNE::FAttributeMap& Attributes) override
+		virtual bool Initialize(TConstArrayView<NNE::FTensorDesc> InputTensorDescs, TConstArrayView<NNE::FTensorDesc> OutputTensorDescs, const NNERuntimeRDGData::Internal::FAttributeMap& Attributes) override
 		{
 			check(InputTensorDescs.Num() == 1);
 			check(OutputTensorDescs.Num() == 1);
@@ -148,14 +148,14 @@ namespace UE::NNERuntimeRDG::Private::Hlsl
 		}
 	};
 
-	bool ValidateSoftmaxOperator(const NNE::FAttributeMap& AttributeMap, TConstArrayView<ENNETensorDataType> InputTypes, TConstArrayView<NNE::FSymbolicTensorShape> InputShapes)
+	bool ValidateSoftmaxOperator(const NNERuntimeRDGData::Internal::FAttributeMap& AttributeMap, TConstArrayView<ENNETensorDataType> InputTypes, TConstArrayView<NNE::FSymbolicTensorShape> InputShapes)
 	{
 		bool bIsValid = true;
 
 		// This matches versions 1, 11, 13 of the Softmax and LogSoftmax operators
 		// [REDACTED_PROFILE_URL]/onnx/blob/main/docs/Changelog.md#Softmax-1
 		FAttributeValidator AttributeValidator;
-		AttributeValidator.AddOptional(TEXT("axis"), ENNEAttributeDataType::Int32);
+		AttributeValidator.AddOptional(TEXT("axis"), ENNERuntimeRDGDataAttributeDataType::Int32);
 		bIsValid &= AttributeValidator.Validate(AttributeMap);
 
 		FInputValidator InputValidator;
diff --git a/Engine/Plugins/Experimental/NNERuntimeRDG/Source/NNERuntimeRDG/Private/Hlsl/NNERuntimeRDGSplit.cpp b/Engine/Plugins/Experimental/NNERuntimeRDG/Source/NNERuntimeRDG/Private/Hlsl/NNERuntimeRDGSplit.cpp
--- a/Engine/Plugins/Experimental/NNERuntimeRDG/Source/NNERuntimeRDG/Private/Hlsl/NNERuntimeRDGSplit.cpp
+++ b/Engine/Plugins/Experimental/NNERuntimeRDG/Source/NNERuntimeRDG/Private/Hlsl/NNERuntimeRDGSplit.cpp
@@ -93,7 +93,7 @@ namespace UE::NNERuntimeRDG::Private::Hlsl
 			return 0;
 		};
 
-		virtual bool Initialize(TConstArrayView<NNE::FTensorDesc> InputTensorDescs, TConstArrayView<NNE::FTensorDesc> OutputTensorDescs, const NNE::FAttributeMap& Attributes) override
+		virtual bool Initialize(TConstArrayView<NNE::FTensorDesc> InputTensorDescs, TConstArrayView<NNE::FTensorDesc> OutputTensorDescs, const NNERuntimeRDGData::Internal::FAttributeMap& Attributes) override
 		{
 			
 			check(InputTensorDescs.Num() >= 1 && InputTensorDescs.Num() <= 2);
@@ -194,19 +194,19 @@ namespace UE::NNERuntimeRDG::Private::Hlsl
 	};
 
 	template<int32 Version>
-	bool ValidateSplitOperator(const NNE::FAttributeMap& AttributeMap, TConstArrayView<ENNETensorDataType> InputTypes, TConstArrayView<NNE::FSymbolicTensorShape> InputShapes)
+	bool ValidateSplitOperator(const NNERuntimeRDGData::Internal::FAttributeMap& AttributeMap, TConstArrayView<ENNETensorDataType> InputTypes, TConstArrayView<NNE::FSymbolicTensorShape> InputShapes)
 	{
 		bool bIsValid = true;
 
 		FAttributeValidator AttributeValidator;
-		AttributeValidator.AddOptional(TEXT("axis"), ENNEAttributeDataType::Int32);
+		AttributeValidator.AddOptional(TEXT("axis"), ENNERuntimeRDGDataAttributeDataType::Int32);
 		if constexpr (Version < 13)
 		{
-			AttributeValidator.AddOptional(TEXT("split"), ENNEAttributeDataType::Int32Array);
+			AttributeValidator.AddOptional(TEXT("split"), ENNERuntimeRDGDataAttributeDataType::Int32Array);
 		}
 		if constexpr (Version >= 18)
 		{
-			AttributeValidator.AddOptional(TEXT("num_outputs"), ENNEAttributeDataType::Int32);
+			AttributeValidator.AddOptional(TEXT("num_outputs"), ENNERuntimeRDGDataAttributeDataType::Int32);
 		}
         bIsValid &= AttributeValidator.Validate(AttributeMap);
 
diff --git a/Engine/Plugins/Experimental/NNERuntimeRDG/Source/NNERuntimeRDG/Private/Hlsl/NNERuntimeRDGSqueeze.cpp b/Engine/Plugins/Experimental/NNERuntimeRDG/Source/NNERuntimeRDG/Private/Hlsl/NNERuntimeRDGSqueeze.cpp
--- a/Engine/Plugins/Experimental/NNERuntimeRDG/Source/NNERuntimeRDG/Private/Hlsl/NNERuntimeRDGSqueeze.cpp
+++ b/Engine/Plugins/Experimental/NNERuntimeRDG/Source/NNERuntimeRDG/Private/Hlsl/NNERuntimeRDGSqueeze.cpp
@@ -57,7 +57,7 @@ namespace UE::NNERuntimeRDG::Private::Hlsl
 			return 0;
 		};
 
-		virtual bool Initialize(TConstArrayView<NNE::FTensorDesc> InputTensorDescs, TConstArrayView<NNE::FTensorDesc> OutputTensorDescs, const NNE::FAttributeMap& Attributes) override
+		virtual bool Initialize(TConstArrayView<NNE::FTensorDesc> InputTensorDescs, TConstArrayView<NNE::FTensorDesc> OutputTensorDescs, const NNERuntimeRDGData::Internal::FAttributeMap& Attributes) override
 		{
 			check(InputTensorDescs.Num() == 1);
 			check(OutputTensorDescs.Num() == 1);
@@ -99,12 +99,12 @@ namespace UE::NNERuntimeRDG::Private::Hlsl
 		}
 	};
 
-	bool ValidateSqueezeOperator(const NNE::FAttributeMap& AttributeMap, TConstArrayView<ENNETensorDataType> InputTypes, TConstArrayView<NNE::FSymbolicTensorShape> InputSqueezes)
+	bool ValidateSqueezeOperator(const NNERuntimeRDGData::Internal::FAttributeMap& AttributeMap, TConstArrayView<ENNETensorDataType> InputTypes, TConstArrayView<NNE::FSymbolicTensorShape> InputSqueezes)
 	{
 		bool bIsValid = true;
 
 		FAttributeValidator AttributeValidator;
-		AttributeValidator.AddRequired(TEXT("axes"), ENNEAttributeDataType::Int32Array);
+		AttributeValidator.AddRequired(TEXT("axes"), ENNERuntimeRDGDataAttributeDataType::Int32Array);
 		bIsValid &= AttributeValidator.Validate(AttributeMap);
 
 		FInputValidator InputValidator;
diff --git a/Engine/Plugins/Experimental/NNERuntimeRDG/Source/NNERuntimeRDG/Private/Hlsl/NNERuntimeRDGTranspose.cpp b/Engine/Plugins/Experimental/NNERuntimeRDG/Source/NNERuntimeRDG/Private/Hlsl/NNERuntimeRDGTranspose.cpp
--- a/Engine/Plugins/Experimental/NNERuntimeRDG/Source/NNERuntimeRDG/Private/Hlsl/NNERuntimeRDGTranspose.cpp
+++ b/Engine/Plugins/Experimental/NNERuntimeRDG/Source/NNERuntimeRDG/Private/Hlsl/NNERuntimeRDGTranspose.cpp
@@ -50,7 +50,7 @@ namespace UE::NNERuntimeRDG::Private::Hlsl
 			return 0;
 		};
 
-		virtual bool Initialize(TConstArrayView<NNE::FTensorDesc> InputTensorDescs, TConstArrayView<NNE::FTensorDesc> OutputTensorDescs, const NNE::FAttributeMap& Attributes) override
+		virtual bool Initialize(TConstArrayView<NNE::FTensorDesc> InputTensorDescs, TConstArrayView<NNE::FTensorDesc> OutputTensorDescs, const NNERuntimeRDGData::Internal::FAttributeMap& Attributes) override
 		{
 			check(InputTensorDescs.Num() == 1);
 			check(OutputTensorDescs.Num() == 1);
@@ -130,12 +130,12 @@ namespace UE::NNERuntimeRDG::Private::Hlsl
 		}
 	};
 
-	bool ValidateTransposeOperator(const NNE::FAttributeMap& AttributeMap, TConstArrayView<ENNETensorDataType> InputTypes, TConstArrayView<NNE::FSymbolicTensorShape> InputTransposes)
+	bool ValidateTransposeOperator(const NNERuntimeRDGData::Internal::FAttributeMap& AttributeMap, TConstArrayView<ENNETensorDataType> InputTypes, TConstArrayView<NNE::FSymbolicTensorShape> InputTransposes)
 	{
 		bool bIsValid = true;
 
 		FAttributeValidator AttributeValidator;
-		AttributeValidator.AddOptional(TEXT("perm"), ENNEAttributeDataType::Int32Array);
+		AttributeValidator.AddOptional(TEXT("perm"), ENNERuntimeRDGDataAttributeDataType::Int32Array);
 		bIsValid &= AttributeValidator.Validate(AttributeMap);
 
 		FInputValidator InputValidator;
diff --git a/Engine/Plugins/Experimental/NNERuntimeRDG/Source/NNERuntimeRDG/Private/Hlsl/NNERuntimeRDGUnsqueeze.cpp b/Engine/Plugins/Experimental/NNERuntimeRDG/Source/NNERuntimeRDG/Private/Hlsl/NNERuntimeRDGUnsqueeze.cpp
--- a/Engine/Plugins/Experimental/NNERuntimeRDG/Source/NNERuntimeRDG/Private/Hlsl/NNERuntimeRDGUnsqueeze.cpp
+++ b/Engine/Plugins/Experimental/NNERuntimeRDG/Source/NNERuntimeRDG/Private/Hlsl/NNERuntimeRDGUnsqueeze.cpp
@@ -49,7 +49,7 @@ namespace UE::NNERuntimeRDG::Private::Hlsl
 			return 0;
 		};
 
-		virtual bool Initialize(TConstArrayView<NNE::FTensorDesc> InputTensorDescs, TConstArrayView<NNE::FTensorDesc> OutputTensorDescs, const NNE::FAttributeMap& Attributes) override
+		virtual bool Initialize(TConstArrayView<NNE::FTensorDesc> InputTensorDescs, TConstArrayView<NNE::FTensorDesc> OutputTensorDescs, const NNERuntimeRDGData::Internal::FAttributeMap& Attributes) override
 		{
 			check(InputTensorDescs.Num() == 1);
 			check(OutputTensorDescs.Num() == 1);
@@ -91,12 +91,12 @@ namespace UE::NNERuntimeRDG::Private::Hlsl
 		}
 	};
 
-	bool ValidateUnsqueezeOperator(const NNE::FAttributeMap& AttributeMap, TConstArrayView<ENNETensorDataType> InputTypes, TConstArrayView<NNE::FSymbolicTensorShape> InputUnsqueezes)
+	bool ValidateUnsqueezeOperator(const NNERuntimeRDGData::Internal::FAttributeMap& AttributeMap, TConstArrayView<ENNETensorDataType> InputTypes, TConstArrayView<NNE::FSymbolicTensorShape> InputUnsqueezes)
 	{
 		bool bIsValid = true;
 
 		FAttributeValidator AttributeValidator;
-		AttributeValidator.AddRequired(TEXT("axes"), ENNEAttributeDataType::Int32Array);
+		AttributeValidator.AddRequired(TEXT("axes"), ENNERuntimeRDGDataAttributeDataType::Int32Array);
 		bIsValid &= AttributeValidator.Validate(AttributeMap);
 
 		FInputValidator InputValidator;
diff --git a/Engine/Plugins/Experimental/NNERuntimeRDG/Source/NNERuntimeRDG/Private/Hlsl/NNERuntimeRDGUpsample.cpp b/Engine/Plugins/Experimental/NNERuntimeRDG/Source/NNERuntimeRDG/Private/Hlsl/NNERuntimeRDGUpsample.cpp
--- a/Engine/Plugins/Experimental/NNERuntimeRDG/Source/NNERuntimeRDG/Private/Hlsl/NNERuntimeRDGUpsample.cpp
+++ b/Engine/Plugins/Experimental/NNERuntimeRDG/Source/NNERuntimeRDG/Private/Hlsl/NNERuntimeRDGUpsample.cpp
@@ -79,7 +79,7 @@ namespace UE::NNERuntimeRDG::Private::Hlsl
 			return 0;
 		};
 
-		virtual bool Initialize(TConstArrayView<NNE::FTensorDesc> InputTensorDescs, TConstArrayView<NNE::FTensorDesc> OutputTensorDescs, const NNE::FAttributeMap& Attributes) override
+		virtual bool Initialize(TConstArrayView<NNE::FTensorDesc> InputTensorDescs, TConstArrayView<NNE::FTensorDesc> OutputTensorDescs, const NNERuntimeRDGData::Internal::FAttributeMap& Attributes) override
 		{
 			check(InputTensorDescs.Num() == 2);
 			check(OutputTensorDescs.Num() == 1);
@@ -147,12 +147,12 @@ namespace UE::NNERuntimeRDG::Private::Hlsl
 		}
 	};
 
-	bool ValidateUpsampleOperator(const NNE::FAttributeMap& AttributeMap, TConstArrayView<ENNETensorDataType> InputTypes, TConstArrayView<NNE::FSymbolicTensorShape> InputShapes)
+	bool ValidateUpsampleOperator(const NNERuntimeRDGData::Internal::FAttributeMap& AttributeMap, TConstArrayView<ENNETensorDataType> InputTypes, TConstArrayView<NNE::FSymbolicTensorShape> InputShapes)
 	{
 		bool bIsValid = true;
 
 		FAttributeValidator AttributeValidator;
-		AttributeValidator.AddOptional(AttrName::Mode, ENNEAttributeDataType::String);
+		AttributeValidator.AddOptional(AttrName::Mode, ENNERuntimeRDGDataAttributeDataType::String);
 		bIsValid &= AttributeValidator.Validate(AttributeMap);
 		
 		FString Mode = AttributeMap.GetValueOrDefault<FString>(AttrName::Mode, AttrValue::Nearest);
diff --git a/Engine/Plugins/Experimental/NNERuntimeRDG/Source/NNERuntimeRDG/Private/NNERuntimeRDGBase.cpp b/Engine/Plugins/Experimental/NNERuntimeRDG/Source/NNERuntimeRDG/Private/NNERuntimeRDGBase.cpp
--- a/Engine/Plugins/Experimental/NNERuntimeRDG/Source/NNERuntimeRDG/Private/NNERuntimeRDGBase.cpp
+++ b/Engine/Plugins/Experimental/NNERuntimeRDG/Source/NNERuntimeRDG/Private/NNERuntimeRDGBase.cpp
@@ -15,7 +15,7 @@ namespace UE::NNERuntimeRDG::Private
 //
 //
 bool AlwaysValidValidationFunction(
-	const NNE::FAttributeMap& AttributeMap, 
+	const NNERuntimeRDGData::Internal::FAttributeMap& AttributeMap, 
 	TConstArrayView<ENNETensorDataType> InputTensorTypes,
 	TConstArrayView<NNE::FSymbolicTensorShape> InputShapes)
 {
@@ -91,7 +91,7 @@ void FInputValidator::AddRequired(int32 TemplateIdx)
 //
 //
 //
-void FAttributeValidator::AddOptional(const FString& Name, ENNEAttributeDataType Type)
+void FAttributeValidator::AddOptional(const FString& Name, ENNERuntimeRDGDataAttributeDataType Type)
 {
 #if DO_CHECK
 	const bool bIsFoundInOptional = (nullptr != OptionalAttributes.FindByPredicate([Name](const FEntry& Other) { return Other.Name == Name; }));
@@ -102,7 +102,7 @@ void FAttributeValidator::AddOptional(const FString& Name, ENNEAttributeDataType
 	OptionalAttributes.Emplace(Name, Type);
 }
 
-void FAttributeValidator::AddRequired(const FString& Name, ENNEAttributeDataType Type)
+void FAttributeValidator::AddRequired(const FString& Name, ENNERuntimeRDGDataAttributeDataType Type)
 {
 #if DO_CHECK
 	const bool bIsFoundInOptional = (nullptr != OptionalAttributes.FindByPredicate([Name](const FEntry& Other) { return Other.Name == Name; }));
@@ -113,14 +113,14 @@ void FAttributeValidator::AddRequired(const FString& Name, ENNEAttributeDataType
 	RequiredAttributes.Emplace(Name, Type);
 }
 
-bool FAttributeValidator::Validate(const NNE::FAttributeMap& AttributesToValidate)
+bool FAttributeValidator::Validate(const NNERuntimeRDGData::Internal::FAttributeMap& AttributesToValidate)
 {
 	bool bAreAttributesValid = true;
 
 	//Verify all required attribute are matching specifications
 	for (int32 Idx = 0; Idx < RequiredAttributes.Num(); ++Idx)
 	{
-		const FNNEAttributeValue* FoundAttribute = AttributesToValidate.GetAttributeValue(RequiredAttributes[Idx].Name);
+		const FNNERuntimeRDGDataAttributeValue* FoundAttribute = AttributesToValidate.GetAttributeValue(RequiredAttributes[Idx].Name);
 		
 		if (FoundAttribute == nullptr)
 		{
@@ -141,7 +141,7 @@ bool FAttributeValidator::Validate(const NNE::FAttributeMap& AttributesToValidat
 	//Verify all optional attribute are matching specifications
 	for (int32 Idx = 0; Idx < OptionalAttributes.Num(); ++Idx)
 	{
-		const FNNEAttributeValue* FoundAttribute = AttributesToValidate.GetAttributeValue(OptionalAttributes[Idx].Name);
+		const FNNERuntimeRDGDataAttributeValue* FoundAttribute = AttributesToValidate.GetAttributeValue(OptionalAttributes[Idx].Name);
 		
 		if ((FoundAttribute != nullptr) && (OptionalAttributes[Idx].Type != FoundAttribute->GetType()))
 		{
diff --git a/Engine/Plugins/Experimental/NNERuntimeRDG/Source/NNERuntimeRDG/Private/NNERuntimeRDGBase.h b/Engine/Plugins/Experimental/NNERuntimeRDG/Source/NNERuntimeRDG/Private/NNERuntimeRDGBase.h
--- a/Engine/Plugins/Experimental/NNERuntimeRDG/Source/NNERuntimeRDG/Private/NNERuntimeRDGBase.h
+++ b/Engine/Plugins/Experimental/NNERuntimeRDG/Source/NNERuntimeRDG/Private/NNERuntimeRDGBase.h
@@ -4,9 +4,9 @@
 
 #include "CoreMinimal.h"
 #include "Containers/ContainersFwd.h"
-#include "NNEAttributeMap.h"
-#include "NNEAttributeValue.h"
 #include "NNEHlslShadersLog.h"
+#include "NNERuntimeRDGDataAttributeMap.h"
+#include "NNERuntimeRDGDataAttributeValue.h"
 #include "NNERuntimeRDGDataFormat.h"
 #ifdef NNE_UTILITIES_AVAILABLE
 #include "NNERuntimeRDGUtilsModelOptimizerInterface.h"
@@ -78,7 +78,7 @@ using FTensorRDGArray = TArray<FTensorRDG, TInlineAllocator<16>>;
 using FTensorRDGRefMap = TMap<int32, FTensorRDGRef>;
 using FIntArray = TArray<int32, TInlineAllocator<16>>;
 
-bool AlwaysValidValidationFunction(const NNE::FAttributeMap& AttributeMap, TConstArrayView<ENNETensorDataType> InputTypes, TConstArrayView<NNE::FSymbolicTensorShape> InputShapes);
+bool AlwaysValidValidationFunction(const NNERuntimeRDGData::Internal::FAttributeMap& AttributeMap, TConstArrayView<ENNETensorDataType> InputTypes, TConstArrayView<NNE::FSymbolicTensorShape> InputShapes);
 
 class FInputValidator
 {
@@ -100,14 +100,14 @@ private:
 class FAttributeValidator
 {
 public:
-	void AddOptional(const FString& Name, ENNEAttributeDataType Type);
-	void AddRequired(const FString& Name, ENNEAttributeDataType Type);
-	bool Validate(const NNE::FAttributeMap& AttributesToValidate);
+	void AddOptional(const FString& Name, ENNERuntimeRDGDataAttributeDataType Type);
+	void AddRequired(const FString& Name, ENNERuntimeRDGDataAttributeDataType Type);
+	bool Validate(const NNERuntimeRDGData::Internal::FAttributeMap& AttributesToValidate);
 
 private:
 	struct FEntry
 	{
-		FEntry(const FString& InName, ENNEAttributeDataType InType)
+		FEntry(const FString& InName, ENNERuntimeRDGDataAttributeDataType InType)
 			: Name(InName), Type(InType)
 		{
 		}
@@ -117,7 +117,7 @@ private:
 		//In the same direction we might only support a range of value for a float (for example
 		//we only support integer but the type is float, or only positive values for an int32)
 		FString Name;
-		ENNEAttributeDataType Type;
+		ENNERuntimeRDGDataAttributeDataType Type;
 	};
 
 	TArray<FEntry> RequiredAttributes;
@@ -176,7 +176,7 @@ class TOperatorRegistryRDG
 public:
 
 	typedef TOperatorType* (*OperatorCreateFunc)();
-	typedef bool (*OperatorValidateFunc)(const NNE::FAttributeMap& AttributeMap, TConstArrayView<ENNETensorDataType> InputTypes, TConstArrayView<NNE::FSymbolicTensorShape> InputShapes);
+	typedef bool (*OperatorValidateFunc)(const NNERuntimeRDGData::Internal::FAttributeMap& AttributeMap, TConstArrayView<ENNETensorDataType> InputTypes, TConstArrayView<NNE::FSymbolicTensorShape> InputShapes);
 
 	static TOperatorRegistryRDG* Get()
 	{
@@ -370,7 +370,7 @@ public:
 		{
 			TArray<ENNETensorDataType> InputTensorTypes;
 			TArray<NNE::FSymbolicTensorShape> InputTensorShapes;
-			NNE::FAttributeMap AttributeMap;
+			NNERuntimeRDGData::Internal::FAttributeMap AttributeMap;
 			
 			for (int32 InputTensorIndex: Format.Operators[Idx].InTensors)
 			{
diff --git a/Engine/Plugins/Experimental/NNERuntimeRDG/Source/NNERuntimeRDG/Private/NNERuntimeRDGHlsl.cpp b/Engine/Plugins/Experimental/NNERuntimeRDG/Source/NNERuntimeRDG/Private/NNERuntimeRDGHlsl.cpp
--- a/Engine/Plugins/Experimental/NNERuntimeRDG/Source/NNERuntimeRDG/Private/NNERuntimeRDGHlsl.cpp
+++ b/Engine/Plugins/Experimental/NNERuntimeRDG/Source/NNERuntimeRDG/Private/NNERuntimeRDGHlsl.cpp
@@ -3,7 +3,6 @@
 #include "NNERuntimeRDGHlsl.h"
 
 #include "Misc/SecureHash.h"
-#include "NNEAttributeMap.h"
 #include "NNEHlslShadersLog.h"
 #include "NNEModelData.h"
 #include "NNERuntimeRDG.h"
diff --git a/Engine/Plugins/Experimental/NNERuntimeRDG/Source/NNERuntimeRDG/Private/NNERuntimeRDGHlslOp.h b/Engine/Plugins/Experimental/NNERuntimeRDG/Source/NNERuntimeRDG/Private/NNERuntimeRDGHlslOp.h
--- a/Engine/Plugins/Experimental/NNERuntimeRDG/Source/NNERuntimeRDG/Private/NNERuntimeRDGHlslOp.h
+++ b/Engine/Plugins/Experimental/NNERuntimeRDG/Source/NNERuntimeRDG/Private/NNERuntimeRDGHlslOp.h
@@ -12,7 +12,7 @@ namespace UE::NNERuntimeRDG::Private::Hlsl
 	struct FOperatorHlsl : public IOperatorRDG, public IPrepareOperator
 	{
 		virtual ~FOperatorHlsl() = default;
-		virtual bool Initialize(TConstArrayView<NNE::FTensorDesc> InputTensorDescs, TConstArrayView<NNE::FTensorDesc> OutputTensorDescs, const NNE::FAttributeMap& Attributes) = 0;
+		virtual bool Initialize(TConstArrayView<NNE::FTensorDesc> InputTensorDescs, TConstArrayView<NNE::FTensorDesc> OutputTensorDescs, const NNERuntimeRDGData::Internal::FAttributeMap& Attributes) = 0;
 		virtual void Dispatch(FRDGBuilder& GraphBuilder, TConstArrayView<FTensorRDGRef> InInputTensors, TConstArrayView<FTensorRDGRef> InOutputTensors) = 0;
 		virtual void OptimizeInputsWeights(TArrayView<FTensorRDGRef> InputWeights) {};
 	};
diff --git a/Engine/Plugins/Experimental/NNERuntimeRDG/Source/NNERuntimeRDG/Private/NNERuntimeRDGModelHlsl.cpp b/Engine/Plugins/Experimental/NNERuntimeRDG/Source/NNERuntimeRDG/Private/NNERuntimeRDGModelHlsl.cpp
--- a/Engine/Plugins/Experimental/NNERuntimeRDG/Source/NNERuntimeRDG/Private/NNERuntimeRDGModelHlsl.cpp
+++ b/Engine/Plugins/Experimental/NNERuntimeRDG/Source/NNERuntimeRDG/Private/NNERuntimeRDGModelHlsl.cpp
@@ -4,6 +4,7 @@
 
 #include "NNEHlslShadersLog.h"
 #include "NNETensor.h"
+#include "NNERuntimeRDGDataAttributeMap.h"
 #include "NNERuntimeRDGHlsl.h"
 #include "NNERuntimeRDGHlslOp.h"
 #include "RenderGraphBuilder.h"
@@ -15,7 +16,7 @@ namespace UE::NNERuntimeRDG::Private::Hlsl
 namespace ModelUtils
 {
 
-FOperatorHlsl* OpCreate(const FOperatorDesc& OpDesc, TConstArrayView<NNE::FTensorDesc> InputTensorDescs, TConstArrayView<NNE::FTensorDesc> OutputTensorDescs, const NNE::FAttributeMap& AttributeMap)
+FOperatorHlsl* OpCreate(const FOperatorDesc& OpDesc, TConstArrayView<NNE::FTensorDesc> InputTensorDescs, TConstArrayView<NNE::FTensorDesc> OutputTensorDescs, const NNERuntimeRDGData::Internal::FAttributeMap& AttributeMap)
 {
 	FOperatorRegistryHlsl::OperatorCreateFunc CreateFn = FOperatorRegistryHlsl::Get()->OpFind(OpDesc);
 
@@ -146,7 +147,7 @@ bool FModelInstance::Init(TConstArrayView<uint8> ModelData)
 		Inputs.Reset();
 		Outputs.Reset();
 		InputsAsWeights.Reset();
-		UE::NNE::FAttributeMap AttributeMap;
+		UE::NNERuntimeRDGData::Internal::FAttributeMap AttributeMap;
 
 		for (int32 InputTensorIndex : Format.Operators[Idx].InTensors)
 		{
diff --git a/Engine/Plugins/Experimental/NNERuntimeRDG/Source/NNERuntimeRDG/Private/Tests/NNERuntimeRDGUnitTestOperatorRegistry.cpp b/Engine/Plugins/Experimental/NNERuntimeRDG/Source/NNERuntimeRDG/Private/Tests/NNERuntimeRDGUnitTestOperatorRegistry.cpp
--- a/Engine/Plugins/Experimental/NNERuntimeRDG/Source/NNERuntimeRDG/Private/Tests/NNERuntimeRDGUnitTestOperatorRegistry.cpp
+++ b/Engine/Plugins/Experimental/NNERuntimeRDG/Source/NNERuntimeRDG/Private/Tests/NNERuntimeRDGUnitTestOperatorRegistry.cpp
@@ -22,7 +22,7 @@ public:
 	{
 		return new FOperatorTestConv();
 	}
-	static bool Validate(const NNE::FAttributeMap& AttributeMap, TConstArrayView<ENNETensorDataType> InputTypes, TConstArrayView<NNE::FSymbolicTensorShape> InputShapes)
+	static bool Validate(const NNERuntimeRDGData::Internal::FAttributeMap& AttributeMap, TConstArrayView<ENNETensorDataType> InputTypes, TConstArrayView<NNE::FSymbolicTensorShape> InputShapes)
 	{
 		return true;
 	}
@@ -37,7 +37,7 @@ public:
 	{
 		return new FOperatorTestClip();
 	}
-	static bool Validate(const NNE::FAttributeMap& AttributeMap, TConstArrayView<ENNETensorDataType> InputTypes, TConstArrayView<NNE::FSymbolicTensorShape> InputShapes)
+	static bool Validate(const NNERuntimeRDGData::Internal::FAttributeMap& AttributeMap, TConstArrayView<ENNETensorDataType> InputTypes, TConstArrayView<NNE::FSymbolicTensorShape> InputShapes)
 	{
 		return true;
 	}
@@ -104,7 +104,7 @@ const int Op##OpName##OpVer##Idx = TestOps.Num() - 1;
 		FRegistryTest::OperatorValidateFunc ValidateFunc = FRegistryTest::Get()->OpFindValidation(TestOps[Idx].OpDesc);
 		UTEST_NOT_NULL(TEXT("OpFindValidation()"), (void *) ValidateFunc);
         UTEST_EQUAL(TEXT("OpFoundMatches"), ValidateFunc, TestOps[Idx].ValidateFunc);
-        UE::NNE::FAttributeMap					Attributes;
+        UE::NNERuntimeRDGData::Internal::FAttributeMap					Attributes;
 		TArray<ENNETensorDataType>				InputTypes;
 		TArray<UE::NNE::FSymbolicTensorShape>	InputShapes;
 		UTEST_TRUE(TEXT("OpValidateCall"), ValidateFunc(Attributes, InputTypes, InputShapes));
diff --git a/Engine/Plugins/Experimental/NNERuntimeRDG/Source/NNERuntimeRDGData/Internal/NNERuntimeRDGDataAttributeDataType.h b/Engine/Plugins/Experimental/NNERuntimeRDG/Source/NNERuntimeRDGData/Internal/NNERuntimeRDGDataAttributeDataType.h
new file mode 100644
--- /dev/null
+++ b/Engine/Plugins/Experimental/NNERuntimeRDG/Source/NNERuntimeRDGData/Internal/NNERuntimeRDGDataAttributeDataType.h
@@ -0,0 +1,26 @@
+// Copyright Epic Games, Inc. All Rights Reserved.
+
+#pragma once
+
+#include "UObject/Object.h"
+
+#include "NNERuntimeRDGDataAttributeDataType.generated.h"
+
+/**
+ * Attribute data types
+ * 
+ * Note: also extend NNERuntimeRDGDataAttributeValueTraits.h for more type support
+ */
+UENUM()
+enum class ENNERuntimeRDGDataAttributeDataType : uint8
+{
+	None,
+	Float,								//!< 32-bit floating number
+	FloatArray,							//!< TArray of 32-bit floating numbers
+	Int32,								//!< 32-bit signed integer
+	Int32Array,							//!< TArray of 32-bit signed integers
+	String,								//!< built-in FString
+	StringArray,						//!< TArray of built-in FString
+	Tensor,								//!< FAttributeTensor
+	TensorArray							//!< TArray of FAttributeTensor
+};
\ No newline at end of file
diff --git a/Engine/Plugins/Experimental/NNERuntimeRDG/Source/NNERuntimeRDGData/Internal/NNERuntimeRDGDataAttributeMap.h b/Engine/Plugins/Experimental/NNERuntimeRDG/Source/NNERuntimeRDGData/Internal/NNERuntimeRDGDataAttributeMap.h
new file mode 100644
--- /dev/null
+++ b/Engine/Plugins/Experimental/NNERuntimeRDG/Source/NNERuntimeRDGData/Internal/NNERuntimeRDGDataAttributeMap.h
@@ -0,0 +1,83 @@
+// Copyright Epic Games, Inc. All Rights Reserved.
+
+#pragma once
+
+#include "NNERuntimeRDGDataAttributeValue.h"
+
+namespace UE::NNERuntimeRDGData::Internal
+{
+
+class FAttributeMap
+{
+public:
+	//Set attribute
+	void SetAttribute(const FString& Name, const FNNERuntimeRDGDataAttributeValue& Value)
+	{
+#if DO_CHECK
+		const bool bIsUnique = (nullptr == Attributes.FindByPredicate([Name](const FEntry& Other) { return Other.Name == Name; }));
+		checkf(bIsUnique, TEXT("Attribute name should be unique"));
+#endif
+		Attributes.Emplace(Name, Value);
+	}
+
+	template<typename T>
+	T GetValue(const FString& Name) const
+	{
+		const FNNERuntimeRDGDataAttributeValue* Value = GetAttributeValue(Name);
+		checkf(Value != nullptr, TEXT("Required attribute %s not found"), *Name);
+		return Value->GetValue<T>();
+	}
+
+	//Query attributes
+	template<typename T>
+	T GetValueOrDefault(const FString& Name, T Default) const
+	{
+		const FNNERuntimeRDGDataAttributeValue* Value = GetAttributeValue(Name);
+		return Value == nullptr ? Default : Value->GetValue<T>();
+	}
+	
+	const FNNERuntimeRDGDataAttributeValue* GetAttributeValue(const FString& Name) const
+	{
+		const FEntry * entry = Attributes.FindByPredicate([Name](const FEntry& Other) { return Other.Name == Name; });
+		if (entry != nullptr)
+		{
+			return &entry->Value;
+		}
+		return nullptr;
+	}
+
+	//Iteration
+	int32 Num() const
+	{
+		return Attributes.Num();
+	}
+
+	const FString& GetName(int32 Idx) const
+	{
+		check(Idx < Num());
+		return Attributes[Idx].Name;
+	}
+
+	const FNNERuntimeRDGDataAttributeValue& GetAttributeValue(int32 Idx) const
+	{
+		check(Idx < Num());
+		return Attributes[Idx].Value;
+	}
+
+private:
+
+	struct FEntry
+	{
+		FEntry(const FString& InName, const FNNERuntimeRDGDataAttributeValue& InValue)
+			: Name(InName), Value(InValue)
+		{
+		}
+		
+		FString Name;
+		FNNERuntimeRDGDataAttributeValue Value;
+	};
+
+	TArray<FEntry> Attributes;
+};
+
+} // namespace UE::NNERuntimeRDGData::Internal
\ No newline at end of file
diff --git a/Engine/Plugins/Experimental/NNERuntimeRDG/Source/NNERuntimeRDGData/Internal/NNERuntimeRDGDataAttributeTensor.h b/Engine/Plugins/Experimental/NNERuntimeRDG/Source/NNERuntimeRDGData/Internal/NNERuntimeRDGDataAttributeTensor.h
new file mode 100644
--- /dev/null
+++ b/Engine/Plugins/Experimental/NNERuntimeRDG/Source/NNERuntimeRDGData/Internal/NNERuntimeRDGDataAttributeTensor.h
@@ -0,0 +1,57 @@
+// Copyright Epic Games, Inc. All Rights Reserved.
+
+#pragma once
+
+#include "Containers/ArrayView.h"
+#include "NNETensor.h"
+#include "NNETypes.h"
+#include "Serialization/Archive.h"
+
+namespace UE::NNERuntimeRDGData::Internal
+{
+	class FAttributeTensor
+	{
+	protected:
+		ENNETensorDataType	DataType;
+		TArray<uint32, TInlineAllocator<NNE::FTensorShape::MaxRank>> Shape;
+		TArray<uint8>		Data;
+
+	public:
+
+		ENNETensorDataType GetDataType() const
+		{
+			return DataType;
+		}
+
+		void FillFTensorWithShapeAndData(NNE::Internal::FTensor& Tensor) const
+		{
+			Tensor.SetShape(NNE::FTensorShape::Make(Shape));
+			Tensor.SetPreparedData(MakeConstArrayView(Data));
+		}
+
+		static FAttributeTensor Make(const NNE::FTensorShape& Shape,
+									 ENNETensorDataType DataType, 
+									 TConstArrayView<uint8> Data)
+		{
+			uint64 Volume = Shape.Volume();
+			check(Volume <= TNumericLimits<uint32>::Max());
+			check(Data.NumBytes() == NNE::GetTensorDataTypeSizeInBytes(DataType) * Volume);
+
+			FAttributeTensor Tensor;
+			Tensor.DataType = DataType;
+			Tensor.Shape.Reset();
+			Tensor.Shape.Append(Shape.GetData());
+			Tensor.Data.Reset();
+			Tensor.Data.Append(Data);
+			return Tensor;
+		}
+
+		friend FArchive& operator<<(FArchive& Ar, FAttributeTensor& Tensor)
+		{
+			Ar << Tensor.DataType;
+			Ar << Tensor.Shape;
+			Ar << Tensor.Data;
+			return Ar;
+		}
+	};
+} // namespace UE::RuntimeRDGData::Internal
diff --git a/Engine/Plugins/Experimental/NNERuntimeRDG/Source/NNERuntimeRDGData/Internal/NNERuntimeRDGDataAttributeValue.h b/Engine/Plugins/Experimental/NNERuntimeRDG/Source/NNERuntimeRDGData/Internal/NNERuntimeRDGDataAttributeValue.h
new file mode 100644
--- /dev/null
+++ b/Engine/Plugins/Experimental/NNERuntimeRDG/Source/NNERuntimeRDGData/Internal/NNERuntimeRDGDataAttributeValue.h
@@ -0,0 +1,56 @@
+// Copyright Epic Games, Inc. All Rights Reserved.
+
+#pragma once
+
+#include "NNERuntimeRDGDataAttributeValueTraits.h"
+#include "Serialization/MemoryWriter.h"
+#include "Serialization/MemoryReader.h"
+#include "UObject/Object.h"
+
+#include "NNERuntimeRDGDataAttributeValue.generated.h"
+
+USTRUCT()
+struct FNNERuntimeRDGDataAttributeValue
+{
+	GENERATED_USTRUCT_BODY()
+
+	FNNERuntimeRDGDataAttributeValue()
+	{
+		Type = ENNERuntimeRDGDataAttributeDataType::None;
+	}
+
+	template<typename T>
+	explicit FNNERuntimeRDGDataAttributeValue(T InValue)
+	{
+		FMemoryWriter writer(Value, /*bIsPersitent =*/ true);
+		writer << InValue;
+
+		Type = TNNERuntimeRDGDataAttributeValueTraits<T>::GetType();
+	}
+
+	template<typename T>
+	T GetValue() const
+	{
+		check((Type == TNNERuntimeRDGDataAttributeValueTraits<T>::GetType()));
+
+		T Result;
+
+		FMemoryReader Reader(Value, /*bIsPersitent =*/ true);
+		Reader << Result;
+
+		return Result;
+	}
+
+	ENNERuntimeRDGDataAttributeDataType GetType() const
+	{
+		return Type;
+	}
+	
+private:
+	
+	UPROPERTY(VisibleAnywhere, Category = "Neural Network Inference")
+	ENNERuntimeRDGDataAttributeDataType	Type;
+	
+	UPROPERTY(VisibleAnywhere, Category = "Neural Network Inference")
+	TArray<uint8>			Value;
+};
\ No newline at end of file
diff --git a/Engine/Plugins/Experimental/NNERuntimeRDG/Source/NNERuntimeRDGData/Internal/NNERuntimeRDGDataAttributeValueTraits.h b/Engine/Plugins/Experimental/NNERuntimeRDG/Source/NNERuntimeRDGData/Internal/NNERuntimeRDGDataAttributeValueTraits.h
new file mode 100644
--- /dev/null
+++ b/Engine/Plugins/Experimental/NNERuntimeRDG/Source/NNERuntimeRDGData/Internal/NNERuntimeRDGDataAttributeValueTraits.h
@@ -0,0 +1,56 @@
+// Copyright Epic Games, Inc. All Rights Reserved.
+
+#pragma once
+
+#include "NNERuntimeRDGDataAttributeDataType.h"
+#include "NNERuntimeRDGDataAttributeTensor.h"
+
+template<typename T> struct TNNERuntimeRDGDataAttributeValueTraits
+{
+	static constexpr ENNERuntimeRDGDataAttributeDataType GetType()
+	{
+		static_assert(!sizeof(T), "Attribute value trait must be specialized for this type.");
+		return ENNERuntimeRDGDataAttributeDataType::None;
+	}
+};
+
+// Attribute specializations
+template<> struct TNNERuntimeRDGDataAttributeValueTraits<float>
+{
+	static constexpr ENNERuntimeRDGDataAttributeDataType GetType() { return ENNERuntimeRDGDataAttributeDataType::Float; }
+};
+
+template<> struct TNNERuntimeRDGDataAttributeValueTraits<int32>
+{
+	static constexpr ENNERuntimeRDGDataAttributeDataType GetType() { return ENNERuntimeRDGDataAttributeDataType::Int32; }
+};
+
+template<> struct TNNERuntimeRDGDataAttributeValueTraits<FString>
+{
+	static constexpr ENNERuntimeRDGDataAttributeDataType GetType() { return ENNERuntimeRDGDataAttributeDataType::String; }
+};
+
+template<> struct TNNERuntimeRDGDataAttributeValueTraits<UE::NNERuntimeRDGData::Internal::FAttributeTensor>
+{
+	static constexpr ENNERuntimeRDGDataAttributeDataType GetType() { return ENNERuntimeRDGDataAttributeDataType::Tensor; }
+};
+
+template<> struct TNNERuntimeRDGDataAttributeValueTraits<TArray<int32>>
+{
+	static constexpr ENNERuntimeRDGDataAttributeDataType GetType() { return ENNERuntimeRDGDataAttributeDataType::Int32Array; }
+};
+
+template<> struct TNNERuntimeRDGDataAttributeValueTraits<TArray<float>>
+{
+	static constexpr ENNERuntimeRDGDataAttributeDataType GetType() { return ENNERuntimeRDGDataAttributeDataType::FloatArray; }
+};
+
+template<> struct TNNERuntimeRDGDataAttributeValueTraits<TArray<FString>>
+{
+	static constexpr ENNERuntimeRDGDataAttributeDataType GetType() { return ENNERuntimeRDGDataAttributeDataType::StringArray; }
+};
+
+template<> struct TNNERuntimeRDGDataAttributeValueTraits<TArray<UE::NNERuntimeRDGData::Internal::FAttributeTensor>>
+{
+	static constexpr ENNERuntimeRDGDataAttributeDataType GetType() { return ENNERuntimeRDGDataAttributeDataType::TensorArray; }
+};
diff --git a/Engine/Plugins/Experimental/NNERuntimeRDG/Source/NNERuntimeRDGData/Internal/NNERuntimeRDGDataFormat.h b/Engine/Plugins/Experimental/NNERuntimeRDG/Source/NNERuntimeRDGData/Internal/NNERuntimeRDGDataFormat.h
--- a/Engine/Plugins/Experimental/NNERuntimeRDG/Source/NNERuntimeRDGData/Internal/NNERuntimeRDGDataFormat.h
+++ b/Engine/Plugins/Experimental/NNERuntimeRDG/Source/NNERuntimeRDGData/Internal/NNERuntimeRDGDataFormat.h
@@ -2,7 +2,7 @@
 
 #pragma once
 
-#include "NNEAttributeValue.h"
+#include "NNERuntimeRDGDataAttributeValue.h"
 #include "NNETypes.h"
 #include "UObject/Object.h"
 
@@ -30,7 +30,7 @@ struct FNNERuntimeRDGDataAttributeDesc
 	FString Name;
 
 	UPROPERTY(VisibleAnywhere, Category = "Neural Network Inference")
-	FNNEAttributeValue Value;
+	FNNERuntimeRDGDataAttributeValue Value;
 };
 
 USTRUCT()
diff --git a/Engine/Plugins/Experimental/NNERuntimeRDG/Source/NNERuntimeRDGUtils/Internal/NNERuntimeRDGUtilsHelpers.h b/Engine/Plugins/Experimental/NNERuntimeRDG/Source/NNERuntimeRDGUtils/Internal/NNERuntimeRDGUtilsHelpers.h
deleted file mode 100644
diff --git a/Engine/Plugins/Experimental/NNERuntimeRDG/Source/NNERuntimeRDGUtils/Internal/NNERuntimeRDGUtilsModelBuilder.h b/Engine/Plugins/Experimental/NNERuntimeRDG/Source/NNERuntimeRDGUtils/Internal/NNERuntimeRDGUtilsModelBuilder.h
deleted file mode 100644
diff --git a/Engine/Plugins/Experimental/NNERuntimeRDG/Source/NNERuntimeRDGUtils/Internal/NNERuntimeRDGUtilsModelBuilderONNX.h b/Engine/Plugins/Experimental/NNERuntimeRDG/Source/NNERuntimeRDGUtils/Internal/NNERuntimeRDGUtilsModelBuilderONNX.h
deleted file mode 100644
diff --git a/Engine/Plugins/Experimental/NNERuntimeRDG/Source/NNERuntimeRDGUtils/Private/NNERuntimeRDGUtilsHelpers.cpp b/Engine/Plugins/Experimental/NNERuntimeRDG/Source/NNERuntimeRDGUtils/Private/NNERuntimeRDGUtilsHelpers.cpp
--- a/Engine/Plugins/Experimental/NNERuntimeRDG/Source/NNERuntimeRDGUtils/Private/NNERuntimeRDGUtilsHelpers.cpp
+++ b/Engine/Plugins/Experimental/NNERuntimeRDG/Source/NNERuntimeRDGUtils/Private/NNERuntimeRDGUtilsHelpers.cpp
@@ -9,7 +9,7 @@ THIRD_PARTY_INCLUDES_START
 THIRD_PARTY_INCLUDES_END
 
 
-namespace UE::NNERuntimeRDGUtils::Internal
+namespace UE::NNERuntimeRDGUtils::Private
 {
 
 TOptional<uint32> GetOpVersionFromOpsetVersion(const FString& OpType, int OpsetVersion)
@@ -23,4 +23,4 @@ TOptional<uint32> GetOpVersionFromOpsetVersion(const FString& OpType, int OpsetV
 	return (uint32) OpSchema->SinceVersion();
 }
 
-}
\ No newline at end of file
+} // namespace UE::NNERuntimeRDGUtils::Private
\ No newline at end of file
diff --git a/Engine/Plugins/Experimental/NNERuntimeRDG/Source/NNERuntimeRDGUtils/Private/NNERuntimeRDGUtilsHelpers.h b/Engine/Plugins/Experimental/NNERuntimeRDG/Source/NNERuntimeRDGUtils/Private/NNERuntimeRDGUtilsHelpers.h
new file mode 100644
--- /dev/null
+++ b/Engine/Plugins/Experimental/NNERuntimeRDG/Source/NNERuntimeRDGUtils/Private/NNERuntimeRDGUtilsHelpers.h
@@ -0,0 +1,11 @@
+// Copyright Epic Games, Inc. All Rights Reserved.
+#pragma once
+
+#include "CoreMinimal.h"
+
+namespace UE::NNERuntimeRDGUtils::Private
+{
+
+TOptional<uint32> GetOpVersionFromOpsetVersion(const FString& OpType, int OpsetVersion);
+
+} // namespace UE::NNERuntimeRDGUtils::Private
\ No newline at end of file
diff --git a/Engine/Plugins/Experimental/NNERuntimeRDG/Source/NNERuntimeRDGUtils/Private/NNERuntimeRDGUtilsModelBuilder.h b/Engine/Plugins/Experimental/NNERuntimeRDG/Source/NNERuntimeRDGUtils/Private/NNERuntimeRDGUtilsModelBuilder.h
new file mode 100644
--- /dev/null
+++ b/Engine/Plugins/Experimental/NNERuntimeRDG/Source/NNERuntimeRDGUtils/Private/NNERuntimeRDGUtilsModelBuilder.h
@@ -0,0 +1,67 @@
+// Copyright Epic Games, Inc. All Rights Reserved.
+
+#pragma once
+
+#include "CoreMinimal.h"
+#include "NNETensor.h"
+#include "NNETypes.h"
+#include "Templates/UniquePtr.h"
+
+struct FNNERuntimeRDGDataAttributeValue;
+
+namespace UE::NNERuntimeRDGUtils::Private
+{
+
+class IModelBuilder
+{
+public:
+
+	enum class EHandleType : uint8
+	{
+		Invalid,
+		Tensor,
+		Operator
+	};
+
+	template<EHandleType Tag>
+	struct THandle
+	{
+		void*		Ptr { nullptr };
+		EHandleType	Type { Tag };
+	};
+
+	typedef struct THandle<EHandleType::Tensor>		FHTensor;
+	typedef struct THandle<EHandleType::Operator>	FHOperator;
+
+	template<EHandleType Tag>
+	THandle<Tag> MakeHandle(void* Ptr)
+	{
+		THandle<Tag> Handle;
+		Handle.Ptr = Ptr;
+		check(Handle.Type != EHandleType::Invalid);
+		return Handle;
+	}
+
+	virtual ~IModelBuilder() = default;
+
+	// Initialize the model builder
+	virtual bool Begin(const FString& GraphName = TEXT("DefaultGraphName")) = 0;
+
+	// Serialize the model to given array
+	virtual bool End(TArray<uint8>& Data) = 0;
+
+	virtual FHTensor AddTensor(const FString& Name, ENNETensorDataType DataType, TArrayView<const int32> Shape) = 0;
+	virtual FHTensor AddConstantTensor(const FString& Name, ENNETensorDataType DataType, TArrayView<const int32> Shape, const void* Data, uint64 DataSize) = 0;
+	virtual FHTensor AddEmptyTensor() = 0;
+	virtual bool AddInput(FHTensor InTensor) = 0;
+	virtual bool AddOutput(FHTensor OutTensor) = 0;
+	virtual FHOperator AddOperator(const FString& Type, const FString& Domain, TOptional<uint32> Version = FNullOpt{0}, const FString& Name = TEXT("")) = 0;
+	virtual bool AddOperatorInput(FHOperator Op, FHTensor Tensor) = 0;
+	virtual bool AddOperatorAttribute(FHOperator Op, const FString& Name, const FNNERuntimeRDGDataAttributeValue& Value) = 0;
+	virtual bool AddOperatorOutput(FHOperator Op, FHTensor Tensor) = 0;
+};
+
+static constexpr TCHAR OnnxDomainName[] = TEXT("Onnx");
+
+} // UE::NNERuntimeRDGUtils::Private
+
diff --git a/Engine/Plugins/Experimental/NNERuntimeRDG/Source/NNERuntimeRDGUtils/Private/NNERuntimeRDGUtilsModelBuilderNNE.cpp b/Engine/Plugins/Experimental/NNERuntimeRDG/Source/NNERuntimeRDGUtils/Private/NNERuntimeRDGUtilsModelBuilderNNE.cpp
--- a/Engine/Plugins/Experimental/NNERuntimeRDG/Source/NNERuntimeRDGUtils/Private/NNERuntimeRDGUtilsModelBuilderNNE.cpp
+++ b/Engine/Plugins/Experimental/NNERuntimeRDG/Source/NNERuntimeRDGUtils/Private/NNERuntimeRDGUtilsModelBuilderNNE.cpp
@@ -2,8 +2,8 @@
 
 #include "NNERuntimeRDGUtilsModelBuilderNNE.h"
 
-#include "NNEAttributeMap.h"
 #include "NNEHlslShadersLog.h"
+#include "NNERuntimeRDGDataAttributeMap.h"
 #include "NNERuntimeRDGDataFormat.h"
 #include "Misc/StringBuilder.h"
 #include "Serialization/MemoryWriter.h"
@@ -13,9 +13,9 @@ namespace UE::NNERuntimeRDGUtils::Private
 
 namespace ModelBuilderNNEHelper
 {
-	int32 NNETensorCast(Internal::IModelBuilder::FHTensor& Handle)
+	int32 NNETensorCast(IModelBuilder::FHTensor& Handle)
 	{
-		if (Handle.Type == Internal::IModelBuilder::EHandleType::Tensor)
+		if (Handle.Type == IModelBuilder::EHandleType::Tensor)
 		{
 			return int32(reinterpret_cast<int64>(Handle.Ptr));
 		}
@@ -23,9 +23,9 @@ namespace ModelBuilderNNEHelper
 		return -1;
 	}
 
-	int32 NNEOperatorCast(Internal::IModelBuilder::FHOperator& Handle)
+	int32 NNEOperatorCast(IModelBuilder::FHOperator& Handle)
 	{
-		if (Handle.Type == Internal::IModelBuilder::EHandleType::Operator)
+		if (Handle.Type == IModelBuilder::EHandleType::Operator)
 		{
 			return int32(reinterpret_cast<int64>(Handle.Ptr));
 		}
@@ -34,7 +34,7 @@ namespace ModelBuilderNNEHelper
 	}
 } //namespace ModelBuilderNNEHelper
 
-class FModelBuilderNNE : public Internal::IModelBuilder
+class FModelBuilderNNE : public IModelBuilder
 {
 public:
 	virtual bool Begin(const FString& Name) override
@@ -181,7 +181,7 @@ public:
 		return true;
 	}
 
-	virtual bool AddOperatorAttribute(FHOperator Op, const FString& Name, const FNNEAttributeValue& Value) override
+	virtual bool AddOperatorAttribute(FHOperator Op, const FString& Name, const FNNERuntimeRDGDataAttributeValue& Value) override
 	{
 		int32 OpIdx = ModelBuilderNNEHelper::NNEOperatorCast(Op);
 
@@ -263,7 +263,7 @@ private:
 	uint32 EmptyTensorCounter = 0;
 };
 
-TUniquePtr<Internal::IModelBuilder> CreateNNEModelBuilder()
+TUniquePtr<IModelBuilder> CreateNNEModelBuilder()
 {
 	return MakeUnique<FModelBuilderNNE>();
 }
diff --git a/Engine/Plugins/Experimental/NNERuntimeRDG/Source/NNERuntimeRDGUtils/Private/NNERuntimeRDGUtilsModelBuilderNNE.h b/Engine/Plugins/Experimental/NNERuntimeRDG/Source/NNERuntimeRDGUtils/Private/NNERuntimeRDGUtilsModelBuilderNNE.h
--- a/Engine/Plugins/Experimental/NNERuntimeRDG/Source/NNERuntimeRDGUtils/Private/NNERuntimeRDGUtilsModelBuilderNNE.h
+++ b/Engine/Plugins/Experimental/NNERuntimeRDG/Source/NNERuntimeRDGUtils/Private/NNERuntimeRDGUtilsModelBuilderNNE.h
@@ -8,7 +8,7 @@
 namespace UE::NNERuntimeRDGUtils::Private
 {
 
-TUniquePtr<NNERuntimeRDGUtils::Internal::IModelBuilder> CreateNNEModelBuilder();
+TUniquePtr<IModelBuilder> CreateNNEModelBuilder();
 
 } // UE::NNERuntimeRDGUtils::Private
 
diff --git a/Engine/Plugins/Experimental/NNERuntimeRDG/Source/NNERuntimeRDGUtils/Private/NNERuntimeRDGUtilsModelBuilderONNX.cpp b/Engine/Plugins/Experimental/NNERuntimeRDG/Source/NNERuntimeRDGUtils/Private/NNERuntimeRDGUtilsModelBuilderONNX.cpp
deleted file mode 100644
diff --git a/Engine/Plugins/Experimental/NNERuntimeRDG/Source/NNERuntimeRDGUtils/Private/NNERuntimeRDGUtilsModelOptimizerNNE.cpp b/Engine/Plugins/Experimental/NNERuntimeRDG/Source/NNERuntimeRDGUtils/Private/NNERuntimeRDGUtilsModelOptimizerNNE.cpp
--- a/Engine/Plugins/Experimental/NNERuntimeRDG/Source/NNERuntimeRDGUtils/Private/NNERuntimeRDGUtilsModelOptimizerNNE.cpp
+++ b/Engine/Plugins/Experimental/NNERuntimeRDG/Source/NNERuntimeRDGUtils/Private/NNERuntimeRDGUtilsModelOptimizerNNE.cpp
@@ -3,8 +3,8 @@
 #include "NNERuntimeRDGUtilsModelOptimizerNNE.h"
 
 #include "Containers/ContainersFwd.h"
-#include "NNEAttributeTensor.h"
 #include "NNEHlslShadersLog.h"
+#include "NNERuntimeRDGDataAttributeTensor.h"
 #include "NNERuntimeRDGDataFormat.h"
 #include "NNERuntimeRDGUtilsModelBuilderNNE.h"
 #include "NNERuntimeRDGUtilsModelOptimizerONNX.h"
@@ -23,8 +23,6 @@ namespace UE::NNERuntimeRDGUtils::Private
 
 namespace ModelOptimizerNNEHelper
 {
-	using namespace UE::NNERuntimeRDGUtils::Internal;
-
 	ENNETensorDataType GetNNETensorTypeFromONNX(int32 DataType)
 	{
 		switch (DataType)
@@ -252,10 +250,10 @@ namespace ModelOptimizerNNEHelper
 		}
 	}
 
-	static bool GetAttributeTensorFromONNXInitializer(const onnx::TensorProto& InTensor, NNE::Internal::FAttributeTensor& OutTensor)
+	static bool GetAttributeTensorFromONNXInitializer(const onnx::TensorProto& InTensor, NNERuntimeRDGData::Internal::FAttributeTensor& OutTensor)
 	{
 		using NNE::FTensorShape;
-		using NNE::Internal::FAttributeTensor;
+		using NNERuntimeRDGData::Internal::FAttributeTensor;
 
 		ENNETensorDataType DataType = ENNETensorDataType::None;
 		TArray<int32> Shape;
@@ -410,12 +408,12 @@ namespace ModelOptimizerNNEHelper
 
 				if (Attribute.type() == onnx::AttributeProto::FLOAT)
 				{
-					Builder->AddOperatorAttribute(Op, AttributeName, FNNEAttributeValue(Attribute.f()));
+					Builder->AddOperatorAttribute(Op, AttributeName, FNNERuntimeRDGDataAttributeValue(Attribute.f()));
 				}
 				else if (Attribute.type() == onnx::AttributeProto::INT)
 				{
 					int32 NNEAttributeValue = ConvertInt32AttributeValueFromONNX((int32)Attribute.i(), AttributeName, NNEOpType);
-					Builder->AddOperatorAttribute(Op, AttributeName, FNNEAttributeValue(NNEAttributeValue));
+					Builder->AddOperatorAttribute(Op, AttributeName, FNNERuntimeRDGDataAttributeValue(NNEAttributeValue));
 				}
 				else if (Attribute.type() == onnx::AttributeProto::INTS)
 				{
@@ -432,11 +430,11 @@ namespace ModelOptimizerNNEHelper
 						Values.Add((int32)ValueClamped);
 					}
 
-					Builder->AddOperatorAttribute(Op, AttributeName, FNNEAttributeValue(Values));
+					Builder->AddOperatorAttribute(Op, AttributeName, FNNERuntimeRDGDataAttributeValue(Values));
 				}
 				else if (Attribute.type() == onnx::AttributeProto::STRING)
 				{
-					Builder->AddOperatorAttribute(Op, AttributeName, FNNEAttributeValue(FString(ANSI_TO_TCHAR(Attribute.s().c_str()))));
+					Builder->AddOperatorAttribute(Op, AttributeName, FNNERuntimeRDGDataAttributeValue(FString(ANSI_TO_TCHAR(Attribute.s().c_str()))));
 				}
 				else if (Attribute.type() == onnx::AttributeProto::STRINGS)
 				{
@@ -447,7 +445,7 @@ namespace ModelOptimizerNNEHelper
 						Values.Add(FString(UTF8_TO_TCHAR(Value.c_str())));
 					}
 
-					Builder->AddOperatorAttribute(Op, AttributeName, FNNEAttributeValue(Values));
+					Builder->AddOperatorAttribute(Op, AttributeName, FNNERuntimeRDGDataAttributeValue(Values));
 				}
 				else if (Attribute.type() == onnx::AttributeProto::FLOATS)
 				{
@@ -458,11 +456,11 @@ namespace ModelOptimizerNNEHelper
 						Values.Add(Value);
 					}
 
-					Builder->AddOperatorAttribute(Op, AttributeName, FNNEAttributeValue(Values));
+					Builder->AddOperatorAttribute(Op, AttributeName, FNNERuntimeRDGDataAttributeValue(Values));
 				}
 				else if (Attribute.type() == onnx::AttributeProto::TENSOR)
 				{
-					using NNE::Internal::FAttributeTensor;
+					using NNERuntimeRDGData::Internal::FAttributeTensor;
 
 					FAttributeTensor Tensor;
 					if (!GetAttributeTensorFromONNXInitializer(Attribute.t(), Tensor))
@@ -470,11 +468,11 @@ namespace ModelOptimizerNNEHelper
 						UE_LOG(LogNNERuntimeRDGHlsl, Error, TEXT("Failed to get data from attribute %s in operator %s"), *AttributeName, *NNEOpName);
 						return false;
 					}
-					Builder->AddOperatorAttribute(Op, AttributeName, FNNEAttributeValue(Tensor));
+					Builder->AddOperatorAttribute(Op, AttributeName, FNNERuntimeRDGDataAttributeValue(Tensor));
 				}
 				else if (Attribute.type() == onnx::AttributeProto::TENSORS)
 				{
-					using NNE::Internal::FAttributeTensor;
+					using NNERuntimeRDGData::Internal::FAttributeTensor;
 
 					TArray<FAttributeTensor> Values;
 
@@ -490,7 +488,7 @@ namespace ModelOptimizerNNEHelper
 						Values.Add(Tensor);
 						TensorIndex++;
 					}
-					Builder->AddOperatorAttribute(Op, AttributeName, FNNEAttributeValue(Values));
+					Builder->AddOperatorAttribute(Op, AttributeName, FNNERuntimeRDGDataAttributeValue(Values));
 				}
 				else
 				{
diff --git a/Engine/Plugins/NNE/NNERuntimeORT/Source/NNERuntimeORT/Private/NNERuntimeORT.cpp b/Engine/Plugins/NNE/NNERuntimeORT/Source/NNERuntimeORT/Private/NNERuntimeORT.cpp
--- a/Engine/Plugins/NNE/NNERuntimeORT/Source/NNERuntimeORT/Private/NNERuntimeORT.cpp
+++ b/Engine/Plugins/NNE/NNERuntimeORT/Source/NNERuntimeORT/Private/NNERuntimeORT.cpp
@@ -4,11 +4,12 @@
 
 #include "Misc/SecureHash.h"
 #include "NNE.h"
-#include "NNEAttributeMap.h"
 #include "NNEModelData.h"
 #include "NNERuntimeORTModel.h"
 #include "NNERuntimeORTModelFormat.h"
 #include "NNERuntimeORTUtils.h"
+#include "Serialization/MemoryReader.h"
+#include "Serialization/MemoryWriter.h"
 
 #if PLATFORM_WINDOWS
 #include "ID3D12DynamicRHI.h"
diff --git a/Engine/Source/Runtime/NNE/Internal/NNEAttributeDataType.h b/Engine/Source/Runtime/NNE/Internal/NNEAttributeDataType.h
deleted file mode 100644
diff --git a/Engine/Source/Runtime/NNE/Internal/NNEAttributeMap.h b/Engine/Source/Runtime/NNE/Internal/NNEAttributeMap.h
deleted file mode 100644
diff --git a/Engine/Source/Runtime/NNE/Internal/NNEAttributeTensor.h b/Engine/Source/Runtime/NNE/Internal/NNEAttributeTensor.h
deleted file mode 100644
diff --git a/Engine/Source/Runtime/NNE/Internal/NNEAttributeValue.h b/Engine/Source/Runtime/NNE/Internal/NNEAttributeValue.h
deleted file mode 100644
diff --git a/Engine/Source/Runtime/NNE/Internal/NNEAttributeValueTraits.h b/Engine/Source/Runtime/NNE/Internal/NNEAttributeValueTraits.h
deleted file mode 100644
