NNERuntimeRDG: Stop forcing c++17

#rnx

diff --git a/Engine/Plugins/Experimental/NNERuntimeRDG/Source/NNERuntimeRDG/NNERuntimeRDG.Build.cs b/Engine/Plugins/Experimental/NNERuntimeRDG/Source/NNERuntimeRDG/NNERuntimeRDG.Build.cs
--- a/Engine/Plugins/Experimental/NNERuntimeRDG/Source/NNERuntimeRDG/NNERuntimeRDG.Build.cs
+++ b/Engine/Plugins/Experimental/NNERuntimeRDG/Source/NNERuntimeRDG/NNERuntimeRDG.Build.cs
@@ -8,8 +8,6 @@ public class NNERuntimeRDG : ModuleRules
 {
 	public NNERuntimeRDG(ReadOnlyTargetRules Target) : base(Target)
 	{
-		CppStandard = CppStandardVersion.Cpp17;
-
 		PublicDependencyModuleNames.AddRange(new string[] 
 		{ 
 			"Core", 
diff --git a/Engine/Plugins/Experimental/NNERuntimeRDG/Source/NNERuntimeRDG/Private/Hlsl/NNERuntimeRDGCumSum.cpp b/Engine/Plugins/Experimental/NNERuntimeRDG/Source/NNERuntimeRDG/Private/Hlsl/NNERuntimeRDGCumSum.cpp
--- a/Engine/Plugins/Experimental/NNERuntimeRDG/Source/NNERuntimeRDG/Private/Hlsl/NNERuntimeRDGCumSum.cpp
+++ b/Engine/Plugins/Experimental/NNERuntimeRDG/Source/NNERuntimeRDG/Private/Hlsl/NNERuntimeRDGCumSum.cpp
@@ -1,5 +1,6 @@
 // Copyright Epic Games, Inc. All Rights Reserved.
 
+#include "NNERuntimeRDGCumSum.h"
 #include "NNEHlslShadersCumSumCS.h"
 
 #include "Helper/NNERuntimeRDGOperatorHelper.h"
