[NNE] Internal Cleanup - step 5.1: remove dependancy from NNERuntimeRDG to NNE/Internal/NNETensor.h

Problem: NNERuntimeRDG is relying on the internal of /NNE
Solution: Brake the dependancy by having the plugin get it's own implementaion of a concrete tensor.

Test:
Win64 editor debug and dev: NNEQA all tests + styletransfer
Win64 package dev : NNEQA all tests + styletransfer

#rnx

diff --git a/Engine/Plugins/Experimental/NNERuntimeRDG/Source/NNEHlslShaders/Internal/NNEHlslShadersGatherCS.h b/Engine/Plugins/Experimental/NNERuntimeRDG/Source/NNEHlslShaders/Internal/NNEHlslShadersGatherCS.h
--- a/Engine/Plugins/Experimental/NNERuntimeRDG/Source/NNEHlslShaders/Internal/NNEHlslShadersGatherCS.h
+++ b/Engine/Plugins/Experimental/NNERuntimeRDG/Source/NNEHlslShaders/Internal/NNEHlslShadersGatherCS.h
@@ -6,7 +6,7 @@
 #include "RenderGraphUtils.h"
 #include "ShaderParameterUtils.h"
 
-namespace UE::NNE::Internal { class FTensor; }
+namespace UE::NNE { class FTensorShape; }
 
 namespace UE::NNEHlslShaders::Internal
 {
@@ -44,7 +44,7 @@ namespace UE::NNEHlslShaders::Internal
 
 		static void ModifyCompilationEnvironment(const FGlobalShaderPermutationParameters& InParameters, FShaderCompilerEnvironment& OutEnvironment);
 
-		static void FillInParameters(int32 Axis, const NNE::Internal::FTensor& Data, const NNE::Internal::FTensor& Indices, FParameters& Parameters);
+		static void FillInParameters(int32 Axis, const NNE::FTensorShape& Data, const NNE::FTensorShape& Indices, FParameters& Parameters);
 
 		static FIntVector GetGroupCount(const FParameters& Parameters);
 	};
diff --git a/Engine/Plugins/Experimental/NNERuntimeRDG/Source/NNEHlslShaders/Internal/NNEHlslShadersGemmCS.h b/Engine/Plugins/Experimental/NNERuntimeRDG/Source/NNEHlslShaders/Internal/NNEHlslShadersGemmCS.h
--- a/Engine/Plugins/Experimental/NNERuntimeRDG/Source/NNEHlslShaders/Internal/NNEHlslShadersGemmCS.h
+++ b/Engine/Plugins/Experimental/NNERuntimeRDG/Source/NNEHlslShaders/Internal/NNEHlslShadersGemmCS.h
@@ -6,7 +6,7 @@
 #include "RenderGraphUtils.h"
 #include "ShaderParameterUtils.h"
 
-namespace UE::NNE::Internal { class FTensor; }
+namespace UE::NNE { class FTensorShape; }
 
 namespace UE::NNEHlslShaders::Internal
 {
@@ -80,9 +80,9 @@ namespace UE::NNEHlslShaders::Internal
 
 		static void ModifyCompilationEnvironment(const FGlobalShaderPermutationParameters& InParameters, FShaderCompilerEnvironment& OutEnvironment);
 
-		static void FillInParameters(float Alpha, float Beta, int32 TransA, int32 TransB, const NNE::Internal::FTensor& InputA, const NNE::Internal::FTensor& InputB, const NNE::Internal::FTensor* InputC, float CScalar, FParameters& Parameters);
+		static void FillInParameters(float Alpha, float Beta, int32 TransA, int32 TransB, const NNE::FTensorShape& InputA, const NNE::FTensorShape& InputB, const NNE::FTensorShape* InputC, float CScalar, FParameters& Parameters);
 
-		static void FillInParametersMatMul(const NNE::Internal::FTensor& InputA, const NNE::Internal::FTensor& InputB, FParameters& Parameters);
+		static void FillInParametersMatMul(const NNE::FTensorShape& InputA, const NNE::FTensorShape& InputB, FParameters& Parameters);
 
 		static FIntVector GetGroupCount(const FParameters& Parameters, EGemmAlgorithm Algorithm, int32 NumStackDimensions);
 		static EGemmAlgorithm GetAlgorithm(const FParameters& Parameters);
diff --git a/Engine/Plugins/Experimental/NNERuntimeRDG/Source/NNEHlslShaders/Private/NNEHlslShadersGatherCS.cpp b/Engine/Plugins/Experimental/NNERuntimeRDG/Source/NNEHlslShaders/Private/NNEHlslShadersGatherCS.cpp
--- a/Engine/Plugins/Experimental/NNERuntimeRDG/Source/NNEHlslShaders/Private/NNEHlslShadersGatherCS.cpp
+++ b/Engine/Plugins/Experimental/NNERuntimeRDG/Source/NNEHlslShaders/Private/NNEHlslShadersGatherCS.cpp
@@ -1,7 +1,7 @@
 // Copyright Epic Games, Inc. All Rights Reserved.
 
 #include "NNEHlslShadersGatherCS.h"
-#include "NNETensor.h"
+#include "NNETypes.h"
 
 namespace UE::NNEHlslShaders::Internal
 {
@@ -16,47 +16,47 @@ namespace UE::NNEHlslShaders::Internal
 	}
 
 	//template <typename DataElementType, typename IndicesElementType>
-	void TGatherCS::FillInParameters(int32 Axis, const NNE::Internal::FTensor& Data, const NNE::Internal::FTensor& Indices, FParameters& Parameters)
+	void TGatherCS::FillInParameters(int32 Axis, const NNE::FTensorShape& Data, const NNE::FTensorShape& Indices, FParameters& Parameters)
 	{
 		Parameters.Axis = Axis;
 
 		TArray<int32> OutputShape;
-		OutputShape.SetNum(Indices.GetShape().Rank() + (Data.GetShape().Rank() - 1));
+		OutputShape.SetNum(Indices.Rank() + (Data.Rank() - 1));
 		int32 index = 0;
 		Parameters.OutputSize = 1;
 		for (int32 i = 0; i < Axis; i++, index++)
 		{
-			OutputShape[index] = Data.GetShape().GetData()[i];
+			OutputShape[index] = Data.GetData()[i];
 			Parameters.OutputSize *= OutputShape[index];
 		}
-		for (int32 i = 0; i < Indices.GetShape().Rank(); i++, index++)
+		for (int32 i = 0; i < Indices.Rank(); i++, index++)
 		{
-			OutputShape[index] = Indices.GetShape().GetData()[i];
+			OutputShape[index] = Indices.GetData()[i];
 			Parameters.OutputSize *= OutputShape[index];
 		}
-		for (int32 i = (Axis + 1); i < Data.GetShape().Rank(); i++, index++)
+		for (int32 i = (Axis + 1); i < Data.Rank(); i++, index++)
 		{
-			OutputShape[index] = Data.GetShape().GetData()[i];
+			OutputShape[index] = Data.GetData()[i];
 			Parameters.OutputSize *= OutputShape[index];
 		}
 
-		Parameters.NumDataDimensions = Data.GetShape().Rank();
-		Parameters.NumIndicesDimensions = Indices.GetShape().Rank();
+		Parameters.NumDataDimensions = Data.Rank();
+		Parameters.NumIndicesDimensions = Indices.Rank();
 
 		int32 DataStride = 1;
-		for (int32 i = Data.GetShape().Rank() - 1; i >= 0; i--)
+		for (int32 i = Data.Rank() - 1; i >= 0; i--)
 		{
 			Parameters.DataStride_IndicesStride_OutputStride_DataSizes[i].X = DataStride;
 			Parameters.OneDivDataStride_OneDivIndicesStride_OneDivOutputStride[i].X = 1.0 / (float)DataStride;
-			DataStride *= Data.GetShape().GetData()[i];
+			DataStride *= Data.GetData()[i];
 		}
 
 		int32 IndicesStride = 1;
-		for (int32 i = Indices.GetShape().Rank() - 1; i >= 0; i--)
+		for (int32 i = Indices.Rank() - 1; i >= 0; i--)
 		{
 			Parameters.DataStride_IndicesStride_OutputStride_DataSizes[i].Y = IndicesStride;
 			Parameters.OneDivDataStride_OneDivIndicesStride_OneDivOutputStride[i].Y = 1.0 / (float)IndicesStride;
-			IndicesStride *= Indices.GetShape().GetData()[i];
+			IndicesStride *= Indices.GetData()[i];
 		}
 
 		int32 OutputStride = 1;
@@ -67,9 +67,9 @@ namespace UE::NNEHlslShaders::Internal
 			OutputStride *= OutputShape[i];
 		}
 
-		for (int32 Idx = 0; Idx < Data.GetShape().Rank(); ++Idx)
+		for (int32 Idx = 0; Idx < Data.Rank(); ++Idx)
 		{
-			Parameters.DataStride_IndicesStride_OutputStride_DataSizes[Idx].W = Data.GetShape().GetData()[Idx];
+			Parameters.DataStride_IndicesStride_OutputStride_DataSizes[Idx].W = Data.GetData()[Idx];
 		}
 	}
 
diff --git a/Engine/Plugins/Experimental/NNERuntimeRDG/Source/NNEHlslShaders/Private/NNEHlslShadersGemmCS.cpp b/Engine/Plugins/Experimental/NNERuntimeRDG/Source/NNEHlslShaders/Private/NNEHlslShadersGemmCS.cpp
--- a/Engine/Plugins/Experimental/NNERuntimeRDG/Source/NNEHlslShaders/Private/NNEHlslShadersGemmCS.cpp
+++ b/Engine/Plugins/Experimental/NNERuntimeRDG/Source/NNEHlslShaders/Private/NNEHlslShadersGemmCS.cpp
@@ -1,8 +1,7 @@
 // Copyright Epic Games, Inc. All Rights Reserved.
 
 #include "NNEHlslShadersGemmCS.h"
-#include "NNE.h"
-#include "NNETensor.h"
+#include "NNETypes.h"
 #include "Algo/Accumulate.h"
 
 namespace UE::NNEHlslShaders::Internal
@@ -18,14 +17,14 @@ namespace UE::NNEHlslShaders::Internal
 		TArray<uint32> StackStrideB;
 	};
 
-	static GemmMatrixParameters GetMatrixParametersMatMul(const NNE::Internal::FTensor& InputA, const NNE::Internal::FTensor& InputB)
+	static GemmMatrixParameters GetMatrixParametersMatMul(const NNE::FTensorShape& InputA, const NNE::FTensorShape& InputB)
 	{
-		check(InputA.GetShape().Rank() != 0);
-		check(InputB.GetShape().Rank() != 0);
+		check(InputA.Rank() != 0);
+		check(InputB.Rank() != 0);
 
 		GemmMatrixParameters Result;
-		int32 NumStackDimensionsA = FMath::Max((int32)InputA.GetShape().Rank() - 2, 0);
-		int32 NumStackDimensionsB = FMath::Max((int32)InputB.GetShape().Rank() - 2, 0);
+		int32 NumStackDimensionsA = FMath::Max((int32)InputA.Rank() - 2, 0);
+		int32 NumStackDimensionsB = FMath::Max((int32)InputB.Rank() - 2, 0);
 		int32 NumStackDimensions = FMath::Max(NumStackDimensionsA, NumStackDimensionsB);
 
 		check(NumStackDimensions <= FGemmConstants::MAX_NUM_STACK_DIMENSIONS);
@@ -36,31 +35,31 @@ namespace UE::NNEHlslShaders::Internal
 			const int32 NumDimToCheck = FMath::Min(NumStackDimensionsA, NumStackDimensionsB);
 			for (int32 i = 0; i < NumDimToCheck; i++)
 			{
-				const uint32 VolumeA = InputA.GetShape().GetData()[NumStackDimensionsA - 1 - i];
-				const uint32 VolumeB = InputB.GetShape().GetData()[NumStackDimensionsB - 1 - i];
+				const uint32 VolumeA = InputA.GetData()[NumStackDimensionsA - 1 - i];
+				const uint32 VolumeB = InputB.GetData()[NumStackDimensionsB - 1 - i];
 
 				check(VolumeA == 1 || VolumeB == 1 || VolumeA == VolumeB);
 			}
 		}
 
-		const bool IsVectorA = InputA.GetShape().Rank() == 1;
-		const bool IsVectorB = InputB.GetShape().Rank() == 1;
+		const bool IsVectorA = InputA.Rank() == 1;
+		const bool IsVectorB = InputB.Rank() == 1;
 
-		Result.M = IsVectorA ? 1 : InputA.GetShape().GetData()[InputA.GetShape().Rank() - 2];
-		Result.N = IsVectorB ? 1 : InputB.GetShape().GetData()[InputB.GetShape().Rank() - 1];
-		Result.K = InputA.GetShape().GetData()[IsVectorA ? 0 : InputA.GetShape().Rank() - 1];
-		check(InputB.GetShape().GetData()[IsVectorB ? 0 : InputB.GetShape().Rank() - 2] == Result.K);
+		Result.M = IsVectorA ? 1 : InputA.GetData()[InputA.Rank() - 2];
+		Result.N = IsVectorB ? 1 : InputB.GetData()[InputB.Rank() - 1];
+		Result.K = InputA.GetData()[IsVectorA ? 0 : InputA.Rank() - 1];
+		check(InputB.GetData()[IsVectorB ? 0 : InputB.Rank() - 2] == Result.K);
 
 		Result.StackShapeA.Init(1, NumStackDimensions);
 		Result.StackShapeB.Init(1, NumStackDimensions);
 
 		for (int32 i = 0; i < NumStackDimensions; i++)
 		{
-			int32 IdxA = InputA.GetShape().Rank() - 3 - i;
-			int32 IdxB = InputB.GetShape().Rank() - 3 - i;
+			int32 IdxA = InputA.Rank() - 3 - i;
+			int32 IdxB = InputB.Rank() - 3 - i;
 
-			Result.StackShapeA[Result.StackShapeA.Num() - 1 - i] = IdxA >= 0 ? InputA.GetShape().GetData()[IdxA] : 1;
-			Result.StackShapeB[Result.StackShapeB.Num() - 1 - i] = IdxB >= 0 ? InputB.GetShape().GetData()[IdxB] : 1;
+			Result.StackShapeA[Result.StackShapeA.Num() - 1 - i] = IdxA >= 0 ? InputA.GetData()[IdxA] : 1;
+			Result.StackShapeB[Result.StackShapeB.Num() - 1 - i] = IdxB >= 0 ? InputB.GetData()[IdxB] : 1;
 		}
 
 		Result.StackStrideA.Init(1, NumStackDimensions);
@@ -81,15 +80,15 @@ namespace UE::NNEHlslShaders::Internal
 		OutEnvironment.SetDefine(TEXT("MAX_NUM_STACK_DIMENSIONS"), FGemmConstants::MAX_NUM_STACK_DIMENSIONS);
 	}
 
-	void TGemmCS::FillInParameters(float Alpha, float Beta, int32 TransA, int32 TransB, const NNE::Internal::FTensor& InputA, const NNE::Internal::FTensor& InputB, const NNE::Internal::FTensor* InputC, float CScalar, TGemmCS::FParameters& Parameters)
+	void TGemmCS::FillInParameters(float Alpha, float Beta, int32 TransA, int32 TransB, const NNE::FTensorShape& InputA, const NNE::FTensorShape& InputB, const NNE::FTensorShape* InputC, float CScalar, TGemmCS::FParameters& Parameters)
 	{
-		check(InputA.GetShape().Rank() == 2);
-		check(InputB.GetShape().Rank() == 2);
+		check(InputA.Rank() == 2);
+		check(InputB.Rank() == 2);
 
-		uint32 M = TransA != 0 ? InputA.GetShape().GetData()[1] : InputA.GetShape().GetData()[0];
-		uint32 K = TransA != 0 ? InputA.GetShape().GetData()[0] : InputA.GetShape().GetData()[1];
-		uint32 N = TransB != 0 ? InputB.GetShape().GetData()[0] : InputB.GetShape().GetData()[1];
-		check(K == (TransB != 0 ? InputB.GetShape().GetData()[1] : InputB.GetShape().GetData()[0]));
+		uint32 M = TransA != 0 ? InputA.GetData()[1] : InputA.GetData()[0];
+		uint32 K = TransA != 0 ? InputA.GetData()[0] : InputA.GetData()[1];
+		uint32 N = TransB != 0 ? InputB.GetData()[0] : InputB.GetData()[1];
+		check(K == (TransB != 0 ? InputB.GetData()[1] : InputB.GetData()[0]));
 
 		Parameters.Alpha = Alpha;
 		Parameters.Beta = Beta;
@@ -105,14 +104,14 @@ namespace UE::NNEHlslShaders::Internal
 		Parameters.CHeight = 0;
 		if (InputC != nullptr)
 		{
-			Parameters.CWidth = InputC->GetShape().Rank() == 0 ? 0 : InputC->GetShape().GetData()[InputC->GetShape().Rank() - 1];;
-			Parameters.CHeight = InputC->GetShape().Rank() == 0 ? 0 : InputC->GetShape().Rank() == 1 ? 1 : InputC->GetShape().GetData()[InputC->GetShape().Rank() - 2];
+			Parameters.CWidth = InputC->Rank() == 0 ? 0 : InputC->GetData()[InputC->Rank() - 1];;
+			Parameters.CHeight = InputC->Rank() == 0 ? 0 : InputC->Rank() == 1 ? 1 : InputC->GetData()[InputC->Rank() - 2];
 		}
 		
 		Parameters.CScalar = CScalar;
 	}
 
-	void TGemmCS::FillInParametersMatMul(const NNE::Internal::FTensor& InputA, const NNE::Internal::FTensor& InputB, TGemmCS::FParameters& Parameters)
+	void TGemmCS::FillInParametersMatMul(const NNE::FTensorShape& InputA, const NNE::FTensorShape& InputB, TGemmCS::FParameters& Parameters)
 	{
 		GemmMatrixParameters Params = GetMatrixParametersMatMul(InputA, InputB);
 
diff --git a/Engine/Plugins/Experimental/NNERuntimeRDG/Source/NNEHlslShaders/Private/NNEHlslShadersInstanceNormalizationCS.cpp b/Engine/Plugins/Experimental/NNERuntimeRDG/Source/NNEHlslShaders/Private/NNEHlslShadersInstanceNormalizationCS.cpp
--- a/Engine/Plugins/Experimental/NNERuntimeRDG/Source/NNEHlslShaders/Private/NNEHlslShadersInstanceNormalizationCS.cpp
+++ b/Engine/Plugins/Experimental/NNERuntimeRDG/Source/NNEHlslShaders/Private/NNEHlslShadersInstanceNormalizationCS.cpp
@@ -2,7 +2,6 @@
 
 #include "NNEHlslShadersInstanceNormalizationCS.h"
 #include "NNE.h"
-#include "NNETensor.h"
 
 namespace UE::NNEHlslShaders::Internal
 {
diff --git a/Engine/Plugins/Experimental/NNERuntimeRDG/Source/NNEHlslShaders/Private/NNEHlslShadersLayerNormalizationCS.cpp b/Engine/Plugins/Experimental/NNERuntimeRDG/Source/NNEHlslShaders/Private/NNEHlslShadersLayerNormalizationCS.cpp
--- a/Engine/Plugins/Experimental/NNERuntimeRDG/Source/NNEHlslShaders/Private/NNEHlslShadersLayerNormalizationCS.cpp
+++ b/Engine/Plugins/Experimental/NNERuntimeRDG/Source/NNEHlslShaders/Private/NNEHlslShadersLayerNormalizationCS.cpp
@@ -2,7 +2,6 @@
 
 #include "NNEHlslShadersLayerNormalizationCS.h"
 #include "NNE.h"
-#include "NNETensor.h"
 
 namespace UE::NNEHlslShaders::Internal
 {
diff --git a/Engine/Plugins/Experimental/NNERuntimeRDG/Source/NNERuntimeRDG/Internal/NNERuntimeRDGAttributes.h b/Engine/Plugins/Experimental/NNERuntimeRDG/Source/NNERuntimeRDG/Internal/NNERuntimeRDGAttributes.h
deleted file mode 100644
diff --git a/Engine/Plugins/Experimental/NNERuntimeRDG/Source/NNERuntimeRDG/Internal/NNERuntimeRDGHelperCast.h b/Engine/Plugins/Experimental/NNERuntimeRDG/Source/NNERuntimeRDG/Internal/NNERuntimeRDGHelperCast.h
deleted file mode 100644
diff --git a/Engine/Plugins/Experimental/NNERuntimeRDG/Source/NNERuntimeRDG/Internal/NNERuntimeRDGHelperConcat.h b/Engine/Plugins/Experimental/NNERuntimeRDG/Source/NNERuntimeRDG/Internal/NNERuntimeRDGHelperConcat.h
deleted file mode 100644
diff --git a/Engine/Plugins/Experimental/NNERuntimeRDG/Source/NNERuntimeRDG/Internal/NNERuntimeRDGHelperElementWiseBinary.h b/Engine/Plugins/Experimental/NNERuntimeRDG/Source/NNERuntimeRDG/Internal/NNERuntimeRDGHelperElementWiseBinary.h
deleted file mode 100644
diff --git a/Engine/Plugins/Experimental/NNERuntimeRDG/Source/NNERuntimeRDG/Internal/NNERuntimeRDGHelperElementWiseUnary.h b/Engine/Plugins/Experimental/NNERuntimeRDG/Source/NNERuntimeRDG/Internal/NNERuntimeRDGHelperElementWiseUnary.h
deleted file mode 100644
diff --git a/Engine/Plugins/Experimental/NNERuntimeRDG/Source/NNERuntimeRDG/Internal/NNERuntimeRDGHelperGather.h b/Engine/Plugins/Experimental/NNERuntimeRDG/Source/NNERuntimeRDG/Internal/NNERuntimeRDGHelperGather.h
deleted file mode 100644
diff --git a/Engine/Plugins/Experimental/NNERuntimeRDG/Source/NNERuntimeRDG/Internal/NNERuntimeRDGHelperSlice.h b/Engine/Plugins/Experimental/NNERuntimeRDG/Source/NNERuntimeRDG/Internal/NNERuntimeRDGHelperSlice.h
deleted file mode 100644
diff --git a/Engine/Plugins/Experimental/NNERuntimeRDG/Source/NNERuntimeRDG/Internal/NNERuntimeRDGHelperTranspose.h b/Engine/Plugins/Experimental/NNERuntimeRDG/Source/NNERuntimeRDG/Internal/NNERuntimeRDGHelperTranspose.h
deleted file mode 100644
diff --git a/Engine/Plugins/Experimental/NNERuntimeRDG/Source/NNERuntimeRDG/Private/Helper/NNERuntimeRDGHelperCast.cpp b/Engine/Plugins/Experimental/NNERuntimeRDG/Source/NNERuntimeRDG/Private/Helper/NNERuntimeRDGHelperCast.cpp
--- a/Engine/Plugins/Experimental/NNERuntimeRDG/Source/NNERuntimeRDG/Private/Helper/NNERuntimeRDGHelperCast.cpp
+++ b/Engine/Plugins/Experimental/NNERuntimeRDG/Source/NNERuntimeRDG/Private/Helper/NNERuntimeRDGHelperCast.cpp
@@ -1,9 +1,9 @@
 // Copyright Epic Games, Inc. All Rights Reserved.
 
 #include "NNERuntimeRDGHelperCast.h"
-#include "NNETensor.h"
+#include "NNERuntimeRDGTensor.h"
 
-namespace UE::NNERuntimeRDG::Internal::CPUHelper::Cast
+namespace UE::NNERuntimeRDG::Private::CPUHelper::Cast
 {
 	template<typename TIn, typename TOut> void ApplyWithResolvedTypes(TConstArrayView<TIn> InArray, TArray<TOut>& OutArray)
 	{
@@ -16,7 +16,7 @@ namespace UE::NNERuntimeRDG::Internal::CPUHelper::Cast
 		}
 	}
 
-	template<typename TOutput> void ApplyWithResolvedOutputType(const NNE::Internal::FTensor& Tensor, NNE::Internal::FTensor& OutputTensor)
+	template<typename TOutput> void ApplyWithResolvedOutputType(const FTensor& Tensor, FTensor& OutputTensor)
 	{
 		ENNETensorDataType SourceType = Tensor.GetDataType();
 		TArray<TOutput> OutputData;
@@ -48,7 +48,7 @@ namespace UE::NNERuntimeRDG::Internal::CPUHelper::Cast
 		OutputTensor.SetPreparedData<TOutput>(OutputData);
 	}
 
-	void Apply(const NNE::Internal::FTensor& Tensor, NNE::Internal::FTensor& OutputTensor)
+	void Apply(const FTensor& Tensor, FTensor& OutputTensor)
 	{
 		//Heuristic to avoid unexpected performance hit. This helper being intended for shape related arithmetic only.
 		static constexpr int32 MaxItemInInputTensors = NNE::FTensorShape::MaxRank * 2;
@@ -82,4 +82,4 @@ namespace UE::NNERuntimeRDG::Internal::CPUHelper::Cast
 			}
 		}
 	}
-} // UE::NNERuntimeRDG::Internal::CPUHelper::Cast
+} // UE::NNERuntimeRDG::Private::CPUHelper::Cast
diff --git a/Engine/Plugins/Experimental/NNERuntimeRDG/Source/NNERuntimeRDG/Private/Helper/NNERuntimeRDGHelperCast.h b/Engine/Plugins/Experimental/NNERuntimeRDG/Source/NNERuntimeRDG/Private/Helper/NNERuntimeRDGHelperCast.h
new file mode 100644
--- /dev/null
+++ b/Engine/Plugins/Experimental/NNERuntimeRDG/Source/NNERuntimeRDG/Private/Helper/NNERuntimeRDGHelperCast.h
@@ -0,0 +1,13 @@
+// Copyright Epic Games, Inc. All Rights Reserved.
+
+#pragma once 
+
+#include "CoreMinimal.h"
+
+namespace UE::NNERuntimeRDG::Private { class FTensor; }
+
+namespace UE::NNERuntimeRDG::Private::CPUHelper::Cast
+{
+	void NNERUNTIMERDG_API Apply(const FTensor& Tensor, FTensor& OutputTensor);
+
+} // UE::NNERuntimeRDG::Private::CPUHelper::Cast
\ No newline at end of file
diff --git a/Engine/Plugins/Experimental/NNERuntimeRDG/Source/NNERuntimeRDG/Private/Helper/NNERuntimeRDGHelperConcat.cpp b/Engine/Plugins/Experimental/NNERuntimeRDG/Source/NNERuntimeRDG/Private/Helper/NNERuntimeRDGHelperConcat.cpp
--- a/Engine/Plugins/Experimental/NNERuntimeRDG/Source/NNERuntimeRDG/Private/Helper/NNERuntimeRDGHelperConcat.cpp
+++ b/Engine/Plugins/Experimental/NNERuntimeRDG/Source/NNERuntimeRDG/Private/Helper/NNERuntimeRDGHelperConcat.cpp
@@ -2,12 +2,12 @@
 
 #include "NNERuntimeRDGHelperConcat.h"
 
-#include "NNETypes.h"
 #include "NNERuntimeRDGTensorIdxIterator.h"
+#include "NNETypes.h"
 
-namespace UE::NNERuntimeRDG::Internal::CPUHelper::Concat
+namespace UE::NNERuntimeRDG::Private::CPUHelper::Concat
 {
-	void Apply(TConstArrayView<NNE::Internal::FTensorRef> InputTensors, NNE::Internal::FTensor& OutputTensor, int32 Axis)
+	void Apply(TConstArrayView<FTensorRef> InputTensors, FTensor& OutputTensor, int32 Axis)
 	{
 		static constexpr int32 MaxItemInOutputTensor = NNE::FTensorShape::MaxRank * 2;
 		
@@ -21,7 +21,7 @@ namespace UE::NNERuntimeRDG::Internal::CPUHelper::Concat
 			return;
 		}
 		
-		for (NNE::Internal::FTensorRef InputTensor : InputTensors)
+		for (FTensorRef InputTensor : InputTensors)
 		{
 			check(InputTensor != nullptr);
 			if (!InputTensor->HasPreparedData())
@@ -38,7 +38,7 @@ namespace UE::NNERuntimeRDG::Internal::CPUHelper::Concat
 
 		OutputData.SetNumUninitialized(OutputTensor.GetVolume());
 		
-		for (NNE::Internal::FTensorRef InputTensor : InputTensors)
+		for (FTensorRef InputTensor : InputTensors)
 		{
 			TConstArrayView<float> InputData = InputTensor->GetPreparedData<float>();
 			Private::TensorIdxIterator itInput(InputTensor->GetShape());
@@ -59,4 +59,4 @@ namespace UE::NNERuntimeRDG::Internal::CPUHelper::Concat
 		OutputTensor.SetPreparedData<float>(OutputData);
 	}
 	
-} // UE::NNERuntimeRDG::Internal::CPUHelper::Concat
+} // UE::NNERuntimeRDG::Private::CPUHelper::Concat
diff --git a/Engine/Plugins/Experimental/NNERuntimeRDG/Source/NNERuntimeRDG/Private/Helper/NNERuntimeRDGHelperConcat.h b/Engine/Plugins/Experimental/NNERuntimeRDG/Source/NNERuntimeRDG/Private/Helper/NNERuntimeRDGHelperConcat.h
new file mode 100644
--- /dev/null
+++ b/Engine/Plugins/Experimental/NNERuntimeRDG/Source/NNERuntimeRDG/Private/Helper/NNERuntimeRDGHelperConcat.h
@@ -0,0 +1,12 @@
+// Copyright Epic Games, Inc. All Rights Reserved.
+
+#pragma once 
+
+#include "CoreMinimal.h"
+#include "NNERuntimeRDGTensor.h"
+
+namespace UE::NNERuntimeRDG::Private::CPUHelper::Concat
+{
+	void NNERUNTIMERDG_API Apply(TConstArrayView<FTensorRef> InputTensors, FTensor& OutputTensor, int32 Axis);
+
+} // UE::NNERuntimeRDG::Private::CPUHelper::Concat
\ No newline at end of file
diff --git a/Engine/Plugins/Experimental/NNERuntimeRDG/Source/NNERuntimeRDG/Private/Helper/NNERuntimeRDGHelperElementWiseBinary.cpp b/Engine/Plugins/Experimental/NNERuntimeRDG/Source/NNERuntimeRDG/Private/Helper/NNERuntimeRDGHelperElementWiseBinary.cpp
--- a/Engine/Plugins/Experimental/NNERuntimeRDG/Source/NNERuntimeRDG/Private/Helper/NNERuntimeRDGHelperElementWiseBinary.cpp
+++ b/Engine/Plugins/Experimental/NNERuntimeRDG/Source/NNERuntimeRDG/Private/Helper/NNERuntimeRDGHelperElementWiseBinary.cpp
@@ -1,13 +1,13 @@
 // Copyright Epic Games, Inc. All Rights Reserved.
 
 #include "NNERuntimeRDGHelperElementWiseBinary.h"
+#include "NNERuntimeRDGTensor.h"
 #include "NNERuntimeRDGTensorIdxIterator.h"
 #include "NNETypes.h"
-#include "NNETensor.h"
 #include "Math/UnrealMathUtility.h"
 #include "MathUtil.h"
 
-namespace UE::NNERuntimeRDG::Internal::CPUHelper::ElementWiseBinary
+namespace UE::NNERuntimeRDG::Private::CPUHelper::ElementWiseBinary
 {
 	using EElementWiseBinaryOperatorType = UE::NNEHlslShaders::Internal::EElementWiseBinaryOperatorType;
 
@@ -37,7 +37,7 @@ namespace UE::NNERuntimeRDG::Internal::CPUHelper::ElementWiseBinary
 	template<> int64 Apply<EElementWiseBinaryOperatorType::Pow, int64>(int64 X, int64 Y) { return (int64)FMath::Pow((float)X, (float)Y); }
 	template<> int64 Apply<EElementWiseBinaryOperatorType::Sub, int64>(int64 X, int64 Y) { return X - Y; }
 
-	template<EElementWiseBinaryOperatorType OpType, typename TData> void Apply(const NNE::Internal::FTensor& LHSTensor, const NNE::Internal::FTensor& RHSTensor, NNE::Internal::FTensor& OutputTensor)
+	template<EElementWiseBinaryOperatorType OpType, typename TData> void Apply(const FTensor& LHSTensor, const FTensor& RHSTensor, FTensor& OutputTensor)
 	{
 		//Heuristic to avoid unexpected performance hit. This helper being intended for shape related arithmetic only.
 		static constexpr int32 MaxItemInInputTensors = NNE::FTensorShape::MaxRank * 2;
@@ -67,7 +67,7 @@ namespace UE::NNERuntimeRDG::Internal::CPUHelper::ElementWiseBinary
 		}
 	}
 
-	template<typename TData> void ApplyResolvedDataType(EElementWiseBinaryOperatorType OpType, const NNE::Internal::FTensor& LHSTensor, const NNE::Internal::FTensor& RHSTensor, NNE::Internal::FTensor& OutputTensor)
+	template<typename TData> void ApplyResolvedDataType(EElementWiseBinaryOperatorType OpType, const FTensor& LHSTensor, const FTensor& RHSTensor, FTensor& OutputTensor)
 	{
 		switch (OpType)
 		{
@@ -97,7 +97,7 @@ namespace UE::NNERuntimeRDG::Internal::CPUHelper::ElementWiseBinary
 		}
 	}
 
-	void Apply(EElementWiseBinaryOperatorType OpType, const NNE::Internal::FTensor& LHSTensor, const NNE::Internal::FTensor& RHSTensor, NNE::Internal::FTensor& OutputTensor)
+	void Apply(EElementWiseBinaryOperatorType OpType, const FTensor& LHSTensor, const FTensor& RHSTensor, FTensor& OutputTensor)
 	{
 		switch (OutputTensor.GetDataType())
 		{
@@ -115,4 +115,4 @@ namespace UE::NNERuntimeRDG::Internal::CPUHelper::ElementWiseBinary
 		}
 	}
 	
-} // UE::NNERuntimeRDG::Internal::CPUHelper::ElementWiseBinary
+} // UE::NNERuntimeRDG::Private::CPUHelper::ElementWiseBinary
diff --git a/Engine/Plugins/Experimental/NNERuntimeRDG/Source/NNERuntimeRDG/Private/Helper/NNERuntimeRDGHelperElementWiseBinary.h b/Engine/Plugins/Experimental/NNERuntimeRDG/Source/NNERuntimeRDG/Private/Helper/NNERuntimeRDGHelperElementWiseBinary.h
new file mode 100644
--- /dev/null
+++ b/Engine/Plugins/Experimental/NNERuntimeRDG/Source/NNERuntimeRDG/Private/Helper/NNERuntimeRDGHelperElementWiseBinary.h
@@ -0,0 +1,14 @@
+// Copyright Epic Games, Inc. All Rights Reserved.
+
+#pragma once 
+
+#include "CoreMinimal.h"
+#include "NNEHlslShadersOperator.h"
+
+namespace UE::NNERuntimeRDG::Private { class FTensor; }
+
+namespace UE::NNERuntimeRDG::Private::CPUHelper::ElementWiseBinary
+{
+	void NNERUNTIMERDG_API Apply(UE::NNEHlslShaders::Internal::EElementWiseBinaryOperatorType OpType, const FTensor& LHSTensor, const FTensor& RHSTensor, FTensor& OutputTensor);
+	
+} // UE::NNERuntimeRDG::Private::CPUHelper::ElementWiseBinary
\ No newline at end of file
diff --git a/Engine/Plugins/Experimental/NNERuntimeRDG/Source/NNERuntimeRDG/Private/Helper/NNERuntimeRDGHelperElementWiseUnary.cpp b/Engine/Plugins/Experimental/NNERuntimeRDG/Source/NNERuntimeRDG/Private/Helper/NNERuntimeRDGHelperElementWiseUnary.cpp
--- a/Engine/Plugins/Experimental/NNERuntimeRDG/Source/NNERuntimeRDG/Private/Helper/NNERuntimeRDGHelperElementWiseUnary.cpp
+++ b/Engine/Plugins/Experimental/NNERuntimeRDG/Source/NNERuntimeRDG/Private/Helper/NNERuntimeRDGHelperElementWiseUnary.cpp
@@ -3,10 +3,10 @@
 #include "NNERuntimeRDGHelperElementWiseUnary.h"
 
 #include "Math/UnrealMathUtility.h"
-#include "NNETensor.h"
+#include "NNERuntimeRDGTensor.h"
 #include "NNETypes.h"
 
-namespace UE::NNERuntimeRDG::Internal::CPUHelper::ElementWiseUnary
+namespace UE::NNERuntimeRDG::Private::CPUHelper::ElementWiseUnary
 {
 	using EElementWiseUnaryOperatorType = UE::NNEHlslShaders::Internal::EElementWiseUnaryOperatorType;
 
@@ -130,7 +130,7 @@ namespace UE::NNERuntimeRDG::Internal::CPUHelper::ElementWiseUnary
 		return Apply<EElementWiseUnaryOperatorType::Tanh>(a * X + b * x3, Alpha, Beta, Gamma);
 	}
 
-	template<EElementWiseUnaryOperatorType OpType> void Apply(const NNE::Internal::FTensor& Tensor, float Alpha, float Beta, float Gamma, NNE::Internal::FTensor& OutputTensor)
+	template<EElementWiseUnaryOperatorType OpType> void Apply(const FTensor& Tensor, float Alpha, float Beta, float Gamma, FTensor& OutputTensor)
 	{
 		//Heuristic to avoid unexpected performance hit. This helper being intended for shape related arithmetic only.
 		static constexpr int32 MaxItemInInputTensors = NNE::FTensorShape::MaxRank * 2;
@@ -167,7 +167,7 @@ namespace UE::NNERuntimeRDG::Internal::CPUHelper::ElementWiseUnary
 		}
 	}
 
-	void Apply(EElementWiseUnaryOperatorType OpType, const NNE::Internal::FTensor& Tensor, float Alpha, float Beta, float Gamma, NNE::Internal::FTensor& OutputTensor)
+	void Apply(EElementWiseUnaryOperatorType OpType, const FTensor& Tensor, float Alpha, float Beta, float Gamma, FTensor& OutputTensor)
 	{
 		switch (OpType)
 		{
@@ -281,4 +281,4 @@ namespace UE::NNERuntimeRDG::Internal::CPUHelper::ElementWiseUnary
 		}
 	}
 	
-} // UE::NNERuntimeRDG::Internal::CPUHelper::ElementWiseUnary
+} // UE::NNERuntimeRDG::Private::CPUHelper::ElementWiseUnary
diff --git a/Engine/Plugins/Experimental/NNERuntimeRDG/Source/NNERuntimeRDG/Private/Helper/NNERuntimeRDGHelperElementWiseUnary.h b/Engine/Plugins/Experimental/NNERuntimeRDG/Source/NNERuntimeRDG/Private/Helper/NNERuntimeRDGHelperElementWiseUnary.h
new file mode 100644
--- /dev/null
+++ b/Engine/Plugins/Experimental/NNERuntimeRDG/Source/NNERuntimeRDG/Private/Helper/NNERuntimeRDGHelperElementWiseUnary.h
@@ -0,0 +1,14 @@
+// Copyright Epic Games, Inc. All Rights Reserved.
+
+#pragma once 
+
+#include "CoreMinimal.h"
+#include "NNEHlslShadersOperator.h"
+
+namespace UE::NNERuntimeRDG::Private { class FTensor; }
+
+namespace UE::NNERuntimeRDG::Private::CPUHelper::ElementWiseUnary
+{
+	void NNERUNTIMERDG_API Apply(UE::NNEHlslShaders::Internal::EElementWiseUnaryOperatorType OpType, const FTensor& Tensor, float Alpha, float Beta, float Gamma, FTensor& OutputTensor);
+
+} // UE::NNERuntimeRDG::Private::CPUHelper::ElementWiseUnary
\ No newline at end of file
diff --git a/Engine/Plugins/Experimental/NNERuntimeRDG/Source/NNERuntimeRDG/Private/Helper/NNERuntimeRDGHelperGather.cpp b/Engine/Plugins/Experimental/NNERuntimeRDG/Source/NNERuntimeRDG/Private/Helper/NNERuntimeRDGHelperGather.cpp
--- a/Engine/Plugins/Experimental/NNERuntimeRDG/Source/NNERuntimeRDG/Private/Helper/NNERuntimeRDGHelperGather.cpp
+++ b/Engine/Plugins/Experimental/NNERuntimeRDG/Source/NNERuntimeRDG/Private/Helper/NNERuntimeRDGHelperGather.cpp
@@ -1,11 +1,12 @@
 // Copyright Epic Games, Inc. All Rights Reserved.
 
 #include "NNERuntimeRDGHelperGather.h"
+#include "NNERuntimeRDGTensor.h"
 #include "NNERuntimeRDGTensorIdxIterator.h"
 
-namespace UE::NNERuntimeRDG::Internal::CPUHelper::Gather
+namespace UE::NNERuntimeRDG::Private::CPUHelper::Gather
 {
-	template <typename TData, typename TInd> void Gather1DResolvedTypes(const NNE::Internal::FTensor& DataTensor, const NNE::Internal::FTensor& IndicesTensor, NNE::Internal::FTensor& OutputTensor)
+	template <typename TData, typename TInd> void Gather1DResolvedTypes(const FTensor& DataTensor, const FTensor& IndicesTensor, FTensor& OutputTensor)
 	{
 		TConstArrayView<TData> DataData = DataTensor.GetPreparedData<TData>();
 		TConstArrayView<TInd> IndicesData = IndicesTensor.GetPreparedData<TInd>();
@@ -22,7 +23,7 @@ namespace UE::NNERuntimeRDG::Internal::CPUHelper::Gather
 		OutputTensor.SetPreparedData<TData>(OutputData);
 	}
 
-	template <typename TInd> void Gather1DResolvedIndices(const NNE::Internal::FTensor& DataTensor, const NNE::Internal::FTensor& IndicesTensor, NNE::Internal::FTensor& OutputTensor)
+	template <typename TInd> void Gather1DResolvedIndices(const FTensor& DataTensor, const FTensor& IndicesTensor, FTensor& OutputTensor)
 	{
 		switch (DataTensor.GetDataType())
 		{
@@ -38,7 +39,7 @@ namespace UE::NNERuntimeRDG::Internal::CPUHelper::Gather
 		}
 	}
 
-	void Apply(const NNE::Internal::FTensor& DataTensor, const NNE::Internal::FTensor& IndicesTensor, int32 Axis, NNE::Internal::FTensor& OutputTensor)
+	void Apply(const FTensor& DataTensor, const FTensor& IndicesTensor, int32 Axis, FTensor& OutputTensor)
 	{
 		static constexpr int32 MaxItemInOutputTensor = NNE::FTensorShape::MaxRank * 2;
 
diff --git a/Engine/Plugins/Experimental/NNERuntimeRDG/Source/NNERuntimeRDG/Private/Helper/NNERuntimeRDGHelperGather.h b/Engine/Plugins/Experimental/NNERuntimeRDG/Source/NNERuntimeRDG/Private/Helper/NNERuntimeRDGHelperGather.h
new file mode 100644
--- /dev/null
+++ b/Engine/Plugins/Experimental/NNERuntimeRDG/Source/NNERuntimeRDG/Private/Helper/NNERuntimeRDGHelperGather.h
@@ -0,0 +1,13 @@
+// Copyright Epic Games, Inc. All Rights Reserved.
+
+#pragma once 
+
+#include "CoreMinimal.h"
+
+namespace UE::NNERuntimeRDG::Private { class FTensor; }
+
+namespace UE::NNERuntimeRDG::Private::CPUHelper::Gather
+{
+	void NNERUNTIMERDG_API Apply(const FTensor& DataTensor, const FTensor& IndicesTensor, int32 Axis, FTensor& OutputTensor);
+
+} // UE::NNERuntimeRDG::Private::CPUHelper::Gather
\ No newline at end of file
diff --git a/Engine/Plugins/Experimental/NNERuntimeRDG/Source/NNERuntimeRDG/Private/Helper/NNERuntimeRDGHelperReshape.h b/Engine/Plugins/Experimental/NNERuntimeRDG/Source/NNERuntimeRDG/Private/Helper/NNERuntimeRDGHelperReshape.h
--- a/Engine/Plugins/Experimental/NNERuntimeRDG/Source/NNERuntimeRDG/Private/Helper/NNERuntimeRDGHelperReshape.h
+++ b/Engine/Plugins/Experimental/NNERuntimeRDG/Source/NNERuntimeRDG/Private/Helper/NNERuntimeRDGHelperReshape.h
@@ -7,7 +7,7 @@
 #include "CoreMinimal.h"
 #include "Misc/EnumerateRange.h"
 #include "NNEHlslShadersLog.h"
-#include "NNETensor.h"
+#include "NNERuntimeRDGTensor.h"
 #include "NNETypes.h"
 
 namespace UE::NNERuntimeRDG::Private::ShapeHelper::Reshape
@@ -15,7 +15,7 @@ namespace UE::NNERuntimeRDG::Private::ShapeHelper::Reshape
 	template<typename DataType, class OutputShapeType >
 	bool ReshapeTensor(
 		const NNE::FTensorShape& InputTensorShape,
-		const NNE::Internal::FTensor& ShapeTensor,
+		const FTensor& ShapeTensor,
 		bool bAllowZero,
 		OutputShapeType& OutShape)
 	{
diff --git a/Engine/Plugins/Experimental/NNERuntimeRDG/Source/NNERuntimeRDG/Private/Helper/NNERuntimeRDGHelperSlice.cpp b/Engine/Plugins/Experimental/NNERuntimeRDG/Source/NNERuntimeRDG/Private/Helper/NNERuntimeRDGHelperSlice.cpp
--- a/Engine/Plugins/Experimental/NNERuntimeRDG/Source/NNERuntimeRDG/Private/Helper/NNERuntimeRDGHelperSlice.cpp
+++ b/Engine/Plugins/Experimental/NNERuntimeRDG/Source/NNERuntimeRDG/Private/Helper/NNERuntimeRDGHelperSlice.cpp
@@ -1,11 +1,12 @@
 // Copyright Epic Games, Inc. All Rights Reserved.
 
 #include "NNERuntimeRDGHelperSlice.h"
+#include "NNERuntimeRDGTensor.h"
 #include "NNERuntimeRDGTensorIdxIterator.h"
 
-namespace UE::NNERuntimeRDG::Internal::CPUHelper::Slice
+namespace UE::NNERuntimeRDG::Private::CPUHelper::Slice
 {
-	template<typename TData> void ApplyResolvedInputType(const NNE::Internal::FTensor& InputTensor, NNE::Internal::FTensor& OutputTensor, TConstArrayView<int32> Starts, TConstArrayView<int32> Steps)
+	template<typename TData> void ApplyResolvedInputType(const FTensor& InputTensor, FTensor& OutputTensor, TConstArrayView<int32> Starts, TConstArrayView<int32> Steps)
 	{
 		check(InputTensor.HasPreparedData());
 		check(InputTensor.GetShape().Rank() == Starts.Num());
@@ -34,7 +35,7 @@ namespace UE::NNERuntimeRDG::Internal::CPUHelper::Slice
 		OutputTensor.SetPreparedData<TData>(OutputData);
 	}
 
-	void Apply(const NNE::Internal::FTensor& InputTensor, NNE::Internal::FTensor& OutputTensor, TConstArrayView<int32> Starts, TConstArrayView<int32> Steps)
+	void Apply(const FTensor& InputTensor, FTensor& OutputTensor, TConstArrayView<int32> Starts, TConstArrayView<int32> Steps)
 	{
 		static constexpr int32 MaxItemInOutputTensor = NNE::FTensorShape::MaxRank * 2;
 
diff --git a/Engine/Plugins/Experimental/NNERuntimeRDG/Source/NNERuntimeRDG/Private/Helper/NNERuntimeRDGHelperSlice.h b/Engine/Plugins/Experimental/NNERuntimeRDG/Source/NNERuntimeRDG/Private/Helper/NNERuntimeRDGHelperSlice.h
new file mode 100644
--- /dev/null
+++ b/Engine/Plugins/Experimental/NNERuntimeRDG/Source/NNERuntimeRDG/Private/Helper/NNERuntimeRDGHelperSlice.h
@@ -0,0 +1,13 @@
+// Copyright Epic Games, Inc. All Rights Reserved.
+
+#pragma once 
+
+#include "CoreMinimal.h"
+
+namespace UE::NNERuntimeRDG::Private { class FTensor; }
+
+namespace UE::NNERuntimeRDG::Private::CPUHelper::Slice
+{
+	void NNERUNTIMERDG_API Apply(const FTensor& InputTensor, FTensor& OutputTensor, TConstArrayView<int32> Starts, TConstArrayView<int32> Steps);
+
+} // UE::NNERuntimeRDG::Private::CPUHelper::Slice
\ No newline at end of file
diff --git a/Engine/Plugins/Experimental/NNERuntimeRDG/Source/NNERuntimeRDG/Private/Helper/NNERuntimeRDGHelperTranspose.cpp b/Engine/Plugins/Experimental/NNERuntimeRDG/Source/NNERuntimeRDG/Private/Helper/NNERuntimeRDGHelperTranspose.cpp
--- a/Engine/Plugins/Experimental/NNERuntimeRDG/Source/NNERuntimeRDG/Private/Helper/NNERuntimeRDGHelperTranspose.cpp
+++ b/Engine/Plugins/Experimental/NNERuntimeRDG/Source/NNERuntimeRDG/Private/Helper/NNERuntimeRDGHelperTranspose.cpp
@@ -1,11 +1,12 @@
 // Copyright Epic Games, Inc. All Rights Reserved.
 
 #include "NNERuntimeRDGHelperTranspose.h"
+#include "NNERuntimeRDGTensor.h"
 #include "NNERuntimeRDGTensorIdxIterator.h"
 
-namespace UE::NNERuntimeRDG::Internal::CPUHelper::Transpose
+namespace UE::NNERuntimeRDG::Private::CPUHelper::Transpose
 {
-	bool Apply(const NNE::Internal::FTensor& Tensor, TConstArrayView<int32> Perms, NNE::Internal::FTensor& OutputTensor)
+	bool Apply(const FTensor& Tensor, TConstArrayView<int32> Perms, FTensor& OutputTensor)
 	{
 		if (!Tensor.HasPreparedData())
 		{
@@ -55,7 +56,7 @@ namespace UE::NNERuntimeRDG::Internal::CPUHelper::Transpose
 		return true;
 	}
 
-	bool TransposePreparedData(NNE::Internal::FTensor& Tensor, TConstArrayView<int32> Perms)
+	bool TransposePreparedData(FTensor& Tensor, TConstArrayView<int32> Perms)
 	{
 		//Transpose in place
 		return Apply(Tensor, Perms, Tensor);
diff --git a/Engine/Plugins/Experimental/NNERuntimeRDG/Source/NNERuntimeRDG/Private/Helper/NNERuntimeRDGHelperTranspose.h b/Engine/Plugins/Experimental/NNERuntimeRDG/Source/NNERuntimeRDG/Private/Helper/NNERuntimeRDGHelperTranspose.h
new file mode 100644
--- /dev/null
+++ b/Engine/Plugins/Experimental/NNERuntimeRDG/Source/NNERuntimeRDG/Private/Helper/NNERuntimeRDGHelperTranspose.h
@@ -0,0 +1,14 @@
+// Copyright Epic Games, Inc. All Rights Reserved.
+
+#pragma once 
+
+#include "CoreMinimal.h"
+
+namespace UE::NNERuntimeRDG::Private { class FTensor; }
+
+namespace UE::NNERuntimeRDG::Private::CPUHelper::Transpose
+{
+	bool NNERUNTIMERDG_API Apply(const FTensor& DataTensor, TConstArrayView<int32> Perms, FTensor& OutputTensor);
+	bool NNERUNTIMERDG_API TransposePreparedData(FTensor& Tensor, TConstArrayView<int32> Perms);
+
+} // UE::NNERuntimeRDG::Private::CPUHelper::Transpose
\ No newline at end of file
diff --git a/Engine/Plugins/Experimental/NNERuntimeRDG/Source/NNERuntimeRDG/Private/Helper/NNERuntimeRDGOperatorHelper.cpp b/Engine/Plugins/Experimental/NNERuntimeRDG/Source/NNERuntimeRDG/Private/Helper/NNERuntimeRDGOperatorHelper.cpp
--- a/Engine/Plugins/Experimental/NNERuntimeRDG/Source/NNERuntimeRDG/Private/Helper/NNERuntimeRDGOperatorHelper.cpp
+++ b/Engine/Plugins/Experimental/NNERuntimeRDG/Source/NNERuntimeRDG/Private/Helper/NNERuntimeRDGOperatorHelper.cpp
@@ -5,7 +5,7 @@
 namespace UE::NNERuntimeRDG::Private::OperatorHelper
 {
 	template<typename Allocator>
-	bool GetInt32ArrayFromConstTensorImpl(TArray<int32, Allocator>& Attr, const NNE::Internal::FTensorRef Tensor)
+	bool GetInt32ArrayFromConstTensorImpl(TArray<int32, Allocator>& Attr, const FTensorRef Tensor)
 	{
 		check(Tensor != nullptr);
 		Attr.Reset();
@@ -34,17 +34,17 @@ namespace UE::NNERuntimeRDG::Private::OperatorHelper
 		return false;
 	}
 	
-	bool GetInt32ArrayFromConstTensor(TArray<int32, TInlineAllocator<NNE::FTensorShape::MaxRank>>& Attr, const NNE::Internal::FTensorRef Tensor)
+	bool GetInt32ArrayFromConstTensor(TArray<int32, TInlineAllocator<NNE::FTensorShape::MaxRank>>& Attr, const FTensorRef Tensor)
 	{
 		return GetInt32ArrayFromConstTensorImpl(Attr, Tensor);
 	}
 	
-	bool GetInt32ArrayFromConstTensor(TArray<int32, TInlineAllocator<2 * NNE::FTensorShape::MaxRank>>& Attr, const NNE::Internal::FTensorRef Tensor)
+	bool GetInt32ArrayFromConstTensor(TArray<int32, TInlineAllocator<2 * NNE::FTensorShape::MaxRank>>& Attr, const FTensorRef Tensor)
 	{
 		return GetInt32ArrayFromConstTensorImpl(Attr, Tensor);
 	}
 
-	bool GetInt32ArrayFromConstTensor(TArray<int32>& Attr, const NNE::Internal::FTensorRef Tensor)
+	bool GetInt32ArrayFromConstTensor(TArray<int32>& Attr, const FTensorRef Tensor)
 	{
 		return GetInt32ArrayFromConstTensorImpl(Attr, Tensor);
 	}
diff --git a/Engine/Plugins/Experimental/NNERuntimeRDG/Source/NNERuntimeRDG/Private/Helper/NNERuntimeRDGOperatorHelper.h b/Engine/Plugins/Experimental/NNERuntimeRDG/Source/NNERuntimeRDG/Private/Helper/NNERuntimeRDGOperatorHelper.h
--- a/Engine/Plugins/Experimental/NNERuntimeRDG/Source/NNERuntimeRDG/Private/Helper/NNERuntimeRDGOperatorHelper.h
+++ b/Engine/Plugins/Experimental/NNERuntimeRDG/Source/NNERuntimeRDG/Private/Helper/NNERuntimeRDGOperatorHelper.h
@@ -3,14 +3,14 @@
 #pragma once
 
 #include "CoreMinimal.h"
-#include "NNETensor.h"
+#include "NNERuntimeRDGTensor.h"
 #include "NNETypes.h"
 
 namespace UE::NNERuntimeRDG::Private::OperatorHelper
 {
-	bool GetInt32ArrayFromConstTensor(TArray<int32, TInlineAllocator<NNE::FTensorShape::MaxRank>>& Attr, const NNE::Internal::FTensorRef Tensor);
-	bool GetInt32ArrayFromConstTensor(TArray<int32, TInlineAllocator<2 * NNE::FTensorShape::MaxRank>>& Attr, const NNE::Internal::FTensorRef Tensor);
-	bool GetInt32ArrayFromConstTensor(TArray<int32>& Attr, const NNE::Internal::FTensorRef Tensor);
+	bool GetInt32ArrayFromConstTensor(TArray<int32, TInlineAllocator<NNE::FTensorShape::MaxRank>>& Attr, const FTensorRef Tensor);
+	bool GetInt32ArrayFromConstTensor(TArray<int32, TInlineAllocator<2 * NNE::FTensorShape::MaxRank>>& Attr, const FTensorRef Tensor);
+	bool GetInt32ArrayFromConstTensor(TArray<int32>& Attr, const FTensorRef Tensor);
 
 } // UE::NNERuntimeRDG::Private::OperatorHelper
 
diff --git a/Engine/Plugins/Experimental/NNERuntimeRDG/Source/NNERuntimeRDG/Private/Helper/NNERuntimeRDGTensorIdxIterator.h b/Engine/Plugins/Experimental/NNERuntimeRDG/Source/NNERuntimeRDG/Private/Helper/NNERuntimeRDGTensorIdxIterator.h
--- a/Engine/Plugins/Experimental/NNERuntimeRDG/Source/NNERuntimeRDG/Private/Helper/NNERuntimeRDGTensorIdxIterator.h
+++ b/Engine/Plugins/Experimental/NNERuntimeRDG/Source/NNERuntimeRDG/Private/Helper/NNERuntimeRDGTensorIdxIterator.h
@@ -3,7 +3,7 @@
 #pragma once
 
 #include "CoreMinimal.h"
-#include "NNETensor.h"
+#include "NNETypes.h"
 
 namespace UE::NNERuntimeRDG::Private
 {
diff --git a/Engine/Plugins/Experimental/NNERuntimeRDG/Source/NNERuntimeRDG/Private/Hlsl/NNERuntimeRDGBatchNormalization.cpp b/Engine/Plugins/Experimental/NNERuntimeRDG/Source/NNERuntimeRDG/Private/Hlsl/NNERuntimeRDGBatchNormalization.cpp
--- a/Engine/Plugins/Experimental/NNERuntimeRDG/Source/NNERuntimeRDG/Private/Hlsl/NNERuntimeRDGBatchNormalization.cpp
+++ b/Engine/Plugins/Experimental/NNERuntimeRDG/Source/NNERuntimeRDG/Private/Hlsl/NNERuntimeRDGBatchNormalization.cpp
@@ -6,8 +6,8 @@
 #include "NNEHlslShadersLog.h"
 #include "NNERuntimeRDGDataAttributeMap.h"
 #include "NNERuntimeRDGHlslHelper.h"
+#include "NNERuntimeRDGTensor.h"
 #include "NNETypes.h"
-#include "NNETensor.h"
 
 namespace UE::NNERuntimeRDG::Private::Hlsl
 {
@@ -27,12 +27,12 @@ namespace UE::NNERuntimeRDG::Private::Hlsl
 
 	public:
 
-		virtual int PrepareOutputs(TConstArrayView<NNE::Internal::FTensorRef> InputTensors, TArrayView<NNE::Internal::FTensorRef> OutputTensors) override
+		virtual int PrepareOutputs(TConstArrayView<FTensorRef> InputTensors, TArrayView<FTensorRef> OutputTensors) override
 		{
 			check(InputTensors.Num() == 5);
 			check(OutputTensors.Num() == 1);
 
-			const NNE::Internal::FTensor& X = *InputTensors[0];
+			const FTensor& X = *InputTensors[0];
 
 			OutputTensors[0]->SetShape(X.GetShape());
 
diff --git a/Engine/Plugins/Experimental/NNERuntimeRDG/Source/NNERuntimeRDG/Private/Hlsl/NNERuntimeRDGCast.cpp b/Engine/Plugins/Experimental/NNERuntimeRDG/Source/NNERuntimeRDG/Private/Hlsl/NNERuntimeRDGCast.cpp
--- a/Engine/Plugins/Experimental/NNERuntimeRDG/Source/NNERuntimeRDG/Private/Hlsl/NNERuntimeRDGCast.cpp
+++ b/Engine/Plugins/Experimental/NNERuntimeRDG/Source/NNERuntimeRDG/Private/Hlsl/NNERuntimeRDGCast.cpp
@@ -2,14 +2,14 @@
 
 #include "NNERuntimeRDGCast.h"
 
+#include "Helper/NNERuntimeRDGHelperCast.h"
 #include "Helper/NNERuntimeRDGLogHelper.h"
 #include "Helper/NNERuntimeRDGOperatorHelper.h"
 #include "NNEHlslShadersCastCS.h"
 #include "NNEHlslShadersLog.h"
 #include "NNEHlslShadersTypeHelper.h"
-#include "NNERuntimeRDGHelperCast.h"
 #include "NNERuntimeRDGHlslHelper.h"
-#include "NNETensor.h"
+#include "NNERuntimeRDGTensor.h"
 #include "NNETypes.h"
 
 namespace UE::NNERuntimeRDG::Private::Hlsl
@@ -25,16 +25,16 @@ namespace UE::NNERuntimeRDG::Private::Hlsl
 		FCast() {}
 		virtual ~FCast() = default;
 
-		virtual int PrepareOutputs(TConstArrayView<NNE::Internal::FTensorRef> InputTensors, TArrayView<NNE::Internal::FTensorRef> OutputTensors) override
+		virtual int PrepareOutputs(TConstArrayView<FTensorRef> InputTensors, TArrayView<FTensorRef> OutputTensors) override
 		{
 			check(InputTensors.Num() == 1);
 			check(OutputTensors.Num() == 1);
 			OutputTensors[0]->SetShape(InputTensors[0]->GetShape());
 
-			const NNE::Internal::FTensor& Input = *InputTensors[0];
-			NNE::Internal::FTensor& Output = *OutputTensors[0];
+			const FTensor& Input = *InputTensors[0];
+			FTensor& Output = *OutputTensors[0];
 
-			Internal::CPUHelper::Cast::Apply(Input, Output);
+			CPUHelper::Cast::Apply(Input, Output);
 
 			bool Has64BitDataType = Input.GetDataType() == ENNETensorDataType::Int64 || Output.GetDataType() == ENNETensorDataType::UInt64;
 
diff --git a/Engine/Plugins/Experimental/NNERuntimeRDG/Source/NNERuntimeRDG/Private/Hlsl/NNERuntimeRDGConcat.cpp b/Engine/Plugins/Experimental/NNERuntimeRDG/Source/NNERuntimeRDG/Private/Hlsl/NNERuntimeRDGConcat.cpp
--- a/Engine/Plugins/Experimental/NNERuntimeRDG/Source/NNERuntimeRDG/Private/Hlsl/NNERuntimeRDGConcat.cpp
+++ b/Engine/Plugins/Experimental/NNERuntimeRDG/Source/NNERuntimeRDG/Private/Hlsl/NNERuntimeRDGConcat.cpp
@@ -2,10 +2,10 @@
 
 #include "NNERuntimeRDGConcat.h"
 
+#include "Helper/NNERuntimeRDGHelperConcat.h"
 #include "NNEHlslShadersLog.h"
-#include "NNETensor.h"
+#include "NNERuntimeRDGTensor.h"
 #include "NNETypes.h"
-#include "NNERuntimeRDGHelperConcat.h"
 #include "RenderGraphUtils.h"
 
 namespace UE::NNERuntimeRDG::Private::Hlsl
@@ -53,7 +53,7 @@ namespace UE::NNERuntimeRDG::Private::Hlsl
 
 	public:
 
-		virtual int PrepareOutputs(TConstArrayView<NNE::Internal::FTensorRef> InputTensors, TArrayView<NNE::Internal::FTensorRef> OutputTensors) override
+		virtual int PrepareOutputs(TConstArrayView<FTensorRef> InputTensors, TArrayView<FTensorRef> OutputTensors) override
 		{
 			check(InputTensors.Num() >= 1);
 			check(OutputTensors.Num() == 1);
@@ -78,7 +78,7 @@ namespace UE::NNERuntimeRDG::Private::Hlsl
 			
 			OutputTensors[0]->SetShape(OutputShape);
 
-			Internal::CPUHelper::Concat::Apply(InputTensors, *OutputTensors[0], Axis);
+			CPUHelper::Concat::Apply(InputTensors, *OutputTensors[0], Axis);
 			
 			return 0;
 		};
diff --git a/Engine/Plugins/Experimental/NNERuntimeRDG/Source/NNERuntimeRDG/Private/Hlsl/NNERuntimeRDGConstant.cpp b/Engine/Plugins/Experimental/NNERuntimeRDG/Source/NNERuntimeRDG/Private/Hlsl/NNERuntimeRDGConstant.cpp
--- a/Engine/Plugins/Experimental/NNERuntimeRDG/Source/NNERuntimeRDG/Private/Hlsl/NNERuntimeRDGConstant.cpp
+++ b/Engine/Plugins/Experimental/NNERuntimeRDG/Source/NNERuntimeRDG/Private/Hlsl/NNERuntimeRDGConstant.cpp
@@ -7,7 +7,7 @@
 #include "NNEHlslShadersLog.h"
 #include "NNERuntimeRDGDataAttributeTensor.h"
 #include "NNERuntimeRDGHlslHelper.h"
-#include "NNETensor.h"
+#include "NNERuntimeRDGTensor.h"
 #include "NNETypes.h"
 
 namespace UE::NNERuntimeRDG::Private::Hlsl
@@ -29,10 +29,9 @@ namespace UE::NNERuntimeRDG::Private::Hlsl
 
 	public:
 
-		virtual int PrepareOutputs(TConstArrayView<NNE::Internal::FTensorRef> InputTensors, TArrayView<NNE::Internal::FTensorRef> OutputTensors) override
+		virtual int PrepareOutputs(TConstArrayView<FTensorRef> InputTensors, TArrayView<FTensorRef> OutputTensors) override
 		{
 			using NNE::FTensorShape;
-			using NNE::Internal::FTensor;
 			using NNERuntimeRDGData::Internal::FAttributeTensor;
 
 			check(InputTensors.Num() == 0);
@@ -80,7 +79,8 @@ namespace UE::NNERuntimeRDG::Private::Hlsl
 												*LogHelper::GetTensorDataTypeName(AttributeTensor.GetDataType()));
 						return -1;
 					}
-					AttributeTensor.FillFTensorWithShapeAndData(Output);
+					Output.SetShape(AttributeTensor.GetShape());
+					Output.SetPreparedData(AttributeTensor.GetData());
 				} break;
 				default:
 					check(false);
diff --git a/Engine/Plugins/Experimental/NNERuntimeRDG/Source/NNERuntimeRDG/Private/Hlsl/NNERuntimeRDGConv.cpp b/Engine/Plugins/Experimental/NNERuntimeRDG/Source/NNERuntimeRDG/Private/Hlsl/NNERuntimeRDGConv.cpp
--- a/Engine/Plugins/Experimental/NNERuntimeRDG/Source/NNERuntimeRDG/Private/Hlsl/NNERuntimeRDGConv.cpp
+++ b/Engine/Plugins/Experimental/NNERuntimeRDG/Source/NNERuntimeRDG/Private/Hlsl/NNERuntimeRDGConv.cpp
@@ -4,8 +4,7 @@
 
 #include "Algo/AnyOf.h"
 #include "HAL/IConsoleManager.h"
-#include "NNETensor.h"
-#include "NNETypes.h"
+#include "Helper/NNERuntimeRDGHelperTranspose.h"
 #include "NNEHlslShadersConvCS.h"
 #include "NNEHlslShadersConvMatmulCS.h"
 #include "NNEHlslShadersConvWinogradInputCS.h"
@@ -15,8 +14,9 @@
 #include "NNEHlslShadersLog.h"
 #include "NNEHlslShadersTypeHelper.h"
 #include "NNERuntimeRDGDataAttributeMap.h"
-#include "NNERuntimeRDGHelperTranspose.h"
 #include "NNERuntimeRDGHlslHelper.h"
+#include "NNERuntimeRDGTensor.h"
+#include "NNETypes.h"
 
 namespace UE::NNERuntimeRDG::Private::Hlsl
 {
@@ -86,7 +86,7 @@ namespace UE::NNERuntimeRDG::Private::Hlsl
 
 	public:
 
-		virtual int PrepareOutputs(TConstArrayView<NNE::Internal::FTensorRef> InputTensors, TArrayView<NNE::Internal::FTensorRef> OutputTensors) override
+		virtual int PrepareOutputs(TConstArrayView<FTensorRef> InputTensors, TArrayView<FTensorRef> OutputTensors) override
 		{
 			using namespace UE::NNEHlslShaders::Internal;
 
@@ -573,7 +573,7 @@ namespace UE::NNERuntimeRDG::Private::Hlsl
 				return;
 
 			//Transpose from CwCiHwWw to HwWwCiCw
-			if (Internal::CPUHelper::Transpose::TransposePreparedData(*Weights, {2,3,1,0} ))
+			if (CPUHelper::Transpose::TransposePreparedData(*Weights, {2,3,1,0} ))
 			{
 				bAreWeightsTransposed = true;
 			}
diff --git a/Engine/Plugins/Experimental/NNERuntimeRDG/Source/NNERuntimeRDG/Private/Hlsl/NNERuntimeRDGConvTranspose.cpp b/Engine/Plugins/Experimental/NNERuntimeRDG/Source/NNERuntimeRDG/Private/Hlsl/NNERuntimeRDGConvTranspose.cpp
--- a/Engine/Plugins/Experimental/NNERuntimeRDG/Source/NNERuntimeRDG/Private/Hlsl/NNERuntimeRDGConvTranspose.cpp
+++ b/Engine/Plugins/Experimental/NNERuntimeRDG/Source/NNERuntimeRDG/Private/Hlsl/NNERuntimeRDGConvTranspose.cpp
@@ -5,7 +5,7 @@
 #include "NNEHlslShadersConvTransposeCS.h"
 #include "NNEHlslShadersLog.h"
 #include "NNERuntimeRDGHlslHelper.h"
-#include "NNETensor.h"
+#include "NNERuntimeRDGTensor.h"
 #include "NNETypes.h"
 
 namespace UE::NNERuntimeRDG::Private::Hlsl
@@ -43,7 +43,7 @@ namespace UE::NNERuntimeRDG::Private::Hlsl
 
 	public:
 
-		virtual int PrepareOutputs(TConstArrayView<NNE::Internal::FTensorRef> InputTensors, TArrayView<NNE::Internal::FTensorRef> OutputTensors) override
+		virtual int PrepareOutputs(TConstArrayView<FTensorRef> InputTensors, TArrayView<FTensorRef> OutputTensors) override
 		{
 			check(InputTensors.Num() >= 2 && InputTensors.Num() <= 3);
 			check(OutputTensors.Num() == 1);
diff --git a/Engine/Plugins/Experimental/NNERuntimeRDG/Source/NNERuntimeRDG/Private/Hlsl/NNERuntimeRDGCumSum.cpp b/Engine/Plugins/Experimental/NNERuntimeRDG/Source/NNERuntimeRDG/Private/Hlsl/NNERuntimeRDGCumSum.cpp
--- a/Engine/Plugins/Experimental/NNERuntimeRDG/Source/NNERuntimeRDG/Private/Hlsl/NNERuntimeRDGCumSum.cpp
+++ b/Engine/Plugins/Experimental/NNERuntimeRDG/Source/NNERuntimeRDG/Private/Hlsl/NNERuntimeRDGCumSum.cpp
@@ -7,7 +7,7 @@
 #include "NNEHlslShadersLog.h"
 #include "NNERuntimeRDGCumSum.h"
 #include "NNERuntimeRDGHlslHelper.h"
-#include "NNETensor.h"
+#include "NNERuntimeRDGTensor.h"
 #include "NNETypes.h"
 #include "RenderGraphUtils.h"
 
@@ -34,12 +34,12 @@ namespace UE::NNERuntimeRDG::Private::Hlsl
 
 	public:
 
-		virtual int PrepareOutputs(TConstArrayView<NNE::Internal::FTensorRef> InputTensors, TArrayView<NNE::Internal::FTensorRef> OutputTensors) override
+		virtual int PrepareOutputs(TConstArrayView<FTensorRef> InputTensors, TArrayView<FTensorRef> OutputTensors) override
 		{
 			check(InputTensors.Num() == 2);
 			check(OutputTensors.Num() == 1);
 
-			const NNE::Internal::FTensorRef AxisTensor = InputTensors[1];
+			const FTensorRef AxisTensor = InputTensors[1];
 			if(!AxisTensor->HasPreparedData())
 			{
 				UE_LOG(LogNNERuntimeRDGHlsl, Warning, TEXT("CumSum: Tensor `axis` (name: %s) must be CPU constant."), *InputTensors[1]->GetName());
@@ -72,7 +72,7 @@ namespace UE::NNERuntimeRDG::Private::Hlsl
 				Axis += InputRank;
 			}
 
-			auto ComputeThreadGroupCount = [Axis=Axis](const NNE::Internal::FTensor& InputTensor) -> FIntVector
+			auto ComputeThreadGroupCount = [Axis=Axis](const FTensor& InputTensor) -> FIntVector
 			{
 				FIntVector OutThreadGroupCount;
 				TConstArrayView<uint32> InputShape = InputTensor.GetShape().GetData();
diff --git a/Engine/Plugins/Experimental/NNERuntimeRDG/Source/NNERuntimeRDG/Private/Hlsl/NNERuntimeRDGDepthToSpace.cpp b/Engine/Plugins/Experimental/NNERuntimeRDG/Source/NNERuntimeRDG/Private/Hlsl/NNERuntimeRDGDepthToSpace.cpp
--- a/Engine/Plugins/Experimental/NNERuntimeRDG/Source/NNERuntimeRDG/Private/Hlsl/NNERuntimeRDGDepthToSpace.cpp
+++ b/Engine/Plugins/Experimental/NNERuntimeRDG/Source/NNERuntimeRDG/Private/Hlsl/NNERuntimeRDGDepthToSpace.cpp
@@ -5,7 +5,7 @@
 #include "NNERuntimeRDGHlslHelper.h"
 #include "NNEHlslShadersLog.h"
 #include "NNEHlslShadersTransposeCS.h"
-#include "NNETensor.h"
+#include "NNERuntimeRDGTensor.h"
 #include "NNETypes.h"
 #include "RenderGraphUtils.h"
 
@@ -51,12 +51,12 @@ namespace UE::NNERuntimeRDG::Private::Hlsl
 
 	public:
 
-		virtual int PrepareOutputs(TConstArrayView<NNE::Internal::FTensorRef> InputTensors, TArrayView<NNE::Internal::FTensorRef> OutputTensors) override
+		virtual int PrepareOutputs(TConstArrayView<FTensorRef> InputTensors, TArrayView<FTensorRef> OutputTensors) override
 		{
 			check(InputTensors.Num() == 1);
 			check(OutputTensors.Num() == 1);
 
-			const NNE::Internal::FTensor& X = *InputTensors[0];
+			const FTensor& X = *InputTensors[0];
 			TConstArrayView<uint32> InputShape = X.GetShape().GetData();
 
 			if(InputShape.Num() != 4)
diff --git a/Engine/Plugins/Experimental/NNERuntimeRDG/Source/NNERuntimeRDG/Private/Hlsl/NNERuntimeRDGDropout.cpp b/Engine/Plugins/Experimental/NNERuntimeRDG/Source/NNERuntimeRDG/Private/Hlsl/NNERuntimeRDGDropout.cpp
--- a/Engine/Plugins/Experimental/NNERuntimeRDG/Source/NNERuntimeRDG/Private/Hlsl/NNERuntimeRDGDropout.cpp
+++ b/Engine/Plugins/Experimental/NNERuntimeRDG/Source/NNERuntimeRDG/Private/Hlsl/NNERuntimeRDGDropout.cpp
@@ -3,7 +3,7 @@
 #include "NNERuntimeRDGDropout.h"
 
 #include "NNEHlslShadersLog.h"
-#include "NNETensor.h"
+#include "NNERuntimeRDGTensor.h"
 #include "NNETypes.h"
 #include "RenderGraphUtils.h"
 
@@ -23,12 +23,12 @@ namespace UE::NNERuntimeRDG::Private::Hlsl
 
 	public:
 
-		virtual int PrepareOutputs(TConstArrayView<NNE::Internal::FTensorRef> InputTensors, TArrayView<NNE::Internal::FTensorRef> OutputTensors) override
+		virtual int PrepareOutputs(TConstArrayView<FTensorRef> InputTensors, TArrayView<FTensorRef> OutputTensors) override
 		{
 			check(InputTensors.Num() == 1);
 			check(OutputTensors.Num() == 1);
 
-			const NNE::Internal::FTensor& X = *InputTensors[0];
+			const FTensor& X = *InputTensors[0];
 
 			OutputTensors[0]->SetShape(X.GetShape());
 			if (X.HasPreparedData())
diff --git a/Engine/Plugins/Experimental/NNERuntimeRDG/Source/NNERuntimeRDG/Private/Hlsl/NNERuntimeRDGElementWiseBinary.cpp b/Engine/Plugins/Experimental/NNERuntimeRDG/Source/NNERuntimeRDG/Private/Hlsl/NNERuntimeRDGElementWiseBinary.cpp
--- a/Engine/Plugins/Experimental/NNERuntimeRDG/Source/NNERuntimeRDG/Private/Hlsl/NNERuntimeRDGElementWiseBinary.cpp
+++ b/Engine/Plugins/Experimental/NNERuntimeRDG/Source/NNERuntimeRDG/Private/Hlsl/NNERuntimeRDGElementWiseBinary.cpp
@@ -2,11 +2,11 @@
 
 #include "NNERuntimeRDGElementWiseBinary.h"
 
+#include "Helper/NNERuntimeRDGHelperElementWiseBinary.h"
 #include "NNEHlslShadersElementWiseBinaryCS.h"
 #include "NNEHlslShadersLog.h"
-#include "NNERuntimeRDGHelperElementWiseBinary.h"
 #include "NNERuntimeRDGHlslHelper.h"
-#include "NNETensor.h"
+#include "NNERuntimeRDGTensor.h"
 #include "NNETypes.h"
 
 namespace UE::NNERuntimeRDG::Private::Hlsl
@@ -30,7 +30,7 @@ namespace UE::NNERuntimeRDG::Private::Hlsl
 
 	public:
 
-		virtual int PrepareOutputs(TConstArrayView<NNE::Internal::FTensorRef> InputTensors, TArrayView<NNE::Internal::FTensorRef> OutputTensors) override
+		virtual int PrepareOutputs(TConstArrayView<FTensorRef> InputTensors, TArrayView<FTensorRef> OutputTensors) override
 		{
 			check(InputTensors.Num() == 2);
 			check(OutputTensors.Num() == 1);
@@ -61,7 +61,7 @@ namespace UE::NNERuntimeRDG::Private::Hlsl
 			
 			OutputTensors[0]->SetShape(OutputShape);
 
-			Internal::CPUHelper::ElementWiseBinary::Apply(OpType, *InputTensors[0], *InputTensors[1], *OutputTensors[0]);
+			CPUHelper::ElementWiseBinary::Apply(OpType, *InputTensors[0], *InputTensors[1], *OutputTensors[0]);
 
 			if (OutputTensors[0]->GetDataType() != ENNETensorDataType::Float && !OutputTensors[0]->HasPreparedData())
 			{
diff --git a/Engine/Plugins/Experimental/NNERuntimeRDG/Source/NNERuntimeRDG/Private/Hlsl/NNERuntimeRDGElementWiseUnary.cpp b/Engine/Plugins/Experimental/NNERuntimeRDG/Source/NNERuntimeRDG/Private/Hlsl/NNERuntimeRDGElementWiseUnary.cpp
--- a/Engine/Plugins/Experimental/NNERuntimeRDG/Source/NNERuntimeRDG/Private/Hlsl/NNERuntimeRDGElementWiseUnary.cpp
+++ b/Engine/Plugins/Experimental/NNERuntimeRDG/Source/NNERuntimeRDG/Private/Hlsl/NNERuntimeRDGElementWiseUnary.cpp
@@ -2,12 +2,12 @@
 
 #include "NNERuntimeRDGElementWiseUnary.h"
 
+#include "Helper/NNERuntimeRDGHelperElementWiseUnary.h"
 #include "NNEHlslShadersElementWiseUnaryCS.h"
 #include "NNEHlslShadersTypeHelper.h"
 #include "NNERuntimeRDGDataAttributeMap.h"
-#include "NNERuntimeRDGHelperElementWiseUnary.h"
 #include "NNERuntimeRDGHlslHelper.h"
-#include "NNETensor.h"
+#include "NNERuntimeRDGTensor.h"
 #include "NNETypes.h"
 
 namespace UE::NNERuntimeRDG::Private::Hlsl
@@ -39,19 +39,19 @@ namespace UE::NNERuntimeRDG::Private::Hlsl
 
 	public:
 
-		virtual int PrepareOutputs(TConstArrayView<NNE::Internal::FTensorRef> InputTensors, TArrayView<NNE::Internal::FTensorRef> OutputTensors) override
+		virtual int PrepareOutputs(TConstArrayView<FTensorRef> InputTensors, TArrayView<FTensorRef> OutputTensors) override
 		{
 			check((OpType == EElementWiseUnaryOperatorType::Clip && Version >= 11) ? (InputTensors.Num() >= 1 && InputTensors.Num() <= 3) : InputTensors.Num() == 1);
 			check(OutputTensors.Num() == 1);
 			
 			OutputTensors[0]->SetShape(InputTensors[0]->GetShape());
 
-			const NNE::Internal::FTensor& X = *InputTensors[0];
+			const FTensor& X = *InputTensors[0];
 
 			if(OpType == EElementWiseUnaryOperatorType::Clip && Version >= 11 && InputTensors.Num() >= 2)
 			{
 				{
-					const NNE::Internal::FTensor& MinTensor = *InputTensors[1];
+					const FTensor& MinTensor = *InputTensors[1];
 					if(MinTensor.HasPreparedData())
 					{
 						Alpha = MinTensor.GetPreparedData<float>()[0];
@@ -60,7 +60,7 @@ namespace UE::NNERuntimeRDG::Private::Hlsl
 
 				if(InputTensors.Num() == 3)
 				{
-					const NNE::Internal::FTensor& MaxTensor = *InputTensors[2];
+					const FTensor& MaxTensor = *InputTensors[2];
 					if(MaxTensor.HasPreparedData())
 					{
 						Beta = MaxTensor.GetPreparedData<float>()[0];
@@ -68,7 +68,7 @@ namespace UE::NNERuntimeRDG::Private::Hlsl
 				}
 			}
 
-			Internal::CPUHelper::ElementWiseUnary::Apply(OpType, X, Alpha, Beta, Gamma, *OutputTensors[0]);
+			CPUHelper::ElementWiseUnary::Apply(OpType, X, Alpha, Beta, Gamma, *OutputTensors[0]);
 			
 			return 0;
 		}
@@ -115,7 +115,7 @@ namespace UE::NNERuntimeRDG::Private::Hlsl
 			if(OpType == EElementWiseUnaryOperatorType::Clip && Version >= 11 && InputTensors.Num() >= 2)
 			{
 				{
-					const NNE::Internal::FTensor& MinTensor = *InputTensors[1];
+					const FTensor& MinTensor = *InputTensors[1];
 					if(!MinTensor.HasPreparedData())
 					{
 						PermutationVector.Set<TElementWiseUnaryCS::FAlphaOnGPU>(true);
@@ -126,7 +126,7 @@ namespace UE::NNERuntimeRDG::Private::Hlsl
 
 				if(InputTensors.Num() == 3)
 				{
-					const NNE::Internal::FTensor& MaxTensor = *InputTensors[2];
+					const FTensor& MaxTensor = *InputTensors[2];
 					if(!MaxTensor.HasPreparedData())
 					{
 						PermutationVector.Set<TElementWiseUnaryCS::FBetaOnGPU>(true);
diff --git a/Engine/Plugins/Experimental/NNERuntimeRDG/Source/NNERuntimeRDG/Private/Hlsl/NNERuntimeRDGElementWiseVariadic.cpp b/Engine/Plugins/Experimental/NNERuntimeRDG/Source/NNERuntimeRDG/Private/Hlsl/NNERuntimeRDGElementWiseVariadic.cpp
--- a/Engine/Plugins/Experimental/NNERuntimeRDG/Source/NNERuntimeRDG/Private/Hlsl/NNERuntimeRDGElementWiseVariadic.cpp
+++ b/Engine/Plugins/Experimental/NNERuntimeRDG/Source/NNERuntimeRDG/Private/Hlsl/NNERuntimeRDGElementWiseVariadic.cpp
@@ -5,7 +5,7 @@
 #include "NNEHlslShadersElementWiseVariadicCS.h"
 #include "NNEHlslShadersLog.h"
 #include "NNERuntimeRDGHlslHelper.h"
-#include "NNETensor.h"
+#include "NNERuntimeRDGTensor.h"
 #include "NNETypes.h"
 
 namespace UE::NNERuntimeRDG::Private::Hlsl
@@ -97,7 +97,7 @@ namespace UE::NNERuntimeRDG::Private::Hlsl
 
 	public:
 
-		virtual int PrepareOutputs(TConstArrayView<NNE::Internal::FTensorRef> InputTensors, TArrayView<NNE::Internal::FTensorRef> OutputTensors) override
+		virtual int PrepareOutputs(TConstArrayView<FTensorRef> InputTensors, TArrayView<FTensorRef> OutputTensors) override
 		{
 			check(InputTensors.Num() > 0);
 			check(OutputTensors.Num() == 1);
diff --git a/Engine/Plugins/Experimental/NNERuntimeRDG/Source/NNERuntimeRDG/Private/Hlsl/NNERuntimeRDGFlatten.cpp b/Engine/Plugins/Experimental/NNERuntimeRDG/Source/NNERuntimeRDG/Private/Hlsl/NNERuntimeRDGFlatten.cpp
--- a/Engine/Plugins/Experimental/NNERuntimeRDG/Source/NNERuntimeRDG/Private/Hlsl/NNERuntimeRDGFlatten.cpp
+++ b/Engine/Plugins/Experimental/NNERuntimeRDG/Source/NNERuntimeRDG/Private/Hlsl/NNERuntimeRDGFlatten.cpp
@@ -3,7 +3,7 @@
 #include "NNERuntimeRDGFlatten.h"
 
 #include "NNEHlslShadersLog.h"
-#include "NNETensor.h"
+#include "NNERuntimeRDGTensor.h"
 #include "NNETypes.h"
 #include "RenderGraphUtils.h"
 
@@ -25,12 +25,12 @@ namespace UE::NNERuntimeRDG::Private::Hlsl
 
 	public:
 
-		virtual int PrepareOutputs(TConstArrayView<NNE::Internal::FTensorRef> InputTensors, TArrayView<NNE::Internal::FTensorRef> OutputTensors) override
+		virtual int PrepareOutputs(TConstArrayView<FTensorRef> InputTensors, TArrayView<FTensorRef> OutputTensors) override
 		{
 			check(InputTensors.Num() == 1);
 			check(OutputTensors.Num() == 1);
 
-			const NNE::Internal::FTensor& X = *InputTensors[0];
+			const FTensor& X = *InputTensors[0];
 			TArray<uint32> OutputShapeData;
 			uint32 InnerDimSize = 1;
 
diff --git a/Engine/Plugins/Experimental/NNERuntimeRDG/Source/NNERuntimeRDG/Private/Hlsl/NNERuntimeRDGGather.cpp b/Engine/Plugins/Experimental/NNERuntimeRDG/Source/NNERuntimeRDG/Private/Hlsl/NNERuntimeRDGGather.cpp
--- a/Engine/Plugins/Experimental/NNERuntimeRDG/Source/NNERuntimeRDG/Private/Hlsl/NNERuntimeRDGGather.cpp
+++ b/Engine/Plugins/Experimental/NNERuntimeRDG/Source/NNERuntimeRDG/Private/Hlsl/NNERuntimeRDGGather.cpp
@@ -2,13 +2,13 @@
 
 #include "NNERuntimeRDGGather.h"
 
+#include "Helper/NNERuntimeRDGHelperGather.h"
 #include "NNEHlslShadersGatherCS.h"
 #include "NNEHlslShadersLog.h"
 #include "NNERuntimeRDGDataAttributeMap.h"
-#include "NNERuntimeRDGHelperGather.h"
 #include "NNERuntimeRDGHlslHelper.h"
+#include "NNERuntimeRDGTensor.h"
 #include "NNETypes.h"
-#include "NNETensor.h"
 
 namespace UE::NNERuntimeRDG::Private::Hlsl
 {
@@ -31,13 +31,13 @@ namespace UE::NNERuntimeRDG::Private::Hlsl
 
 	public:
 
-		virtual int PrepareOutputs(TConstArrayView<NNE::Internal::FTensorRef> InputTensors, TArrayView<NNE::Internal::FTensorRef> OutputTensors) override
+		virtual int PrepareOutputs(TConstArrayView<FTensorRef> InputTensors, TArrayView<FTensorRef> OutputTensors) override
 		{
 			check(InputTensors.Num() == 2)
 			check(OutputTensors.Num() == 1)
 
-			const NNE::Internal::FTensor& DataTensor = *InputTensors[0];
-			const NNE::Internal::FTensor& IndicesTensor = *InputTensors[1];
+			const FTensor& DataTensor = *InputTensors[0];
+			const FTensor& IndicesTensor = *InputTensors[1];
 			const NNE::FTensorShape& DataShape = DataTensor.GetShape();
 			const NNE::FTensorShape& IndicesShape = IndicesTensor.GetShape();
 
@@ -65,7 +65,7 @@ namespace UE::NNERuntimeRDG::Private::Hlsl
 
 			OutputTensors[0]->SetShape(OutputShape);
 
-			Internal::CPUHelper::Gather::Apply(DataTensor, IndicesTensor, Axis, *OutputTensors[0]);
+			CPUHelper::Gather::Apply(DataTensor, IndicesTensor, Axis, *OutputTensors[0]);
 
 			return 0;
 		};
@@ -114,7 +114,7 @@ namespace UE::NNERuntimeRDG::Private::Hlsl
 
 			// Set parameters
 			TGatherCS::FParameters* Parameters = GraphBuilder.AllocParameters<TGatherCS::FParameters>();
-			TGatherCS::FillInParameters(Axis, Data, Indices, *Parameters);
+			TGatherCS::FillInParameters(Axis, Data.GetShape(), Indices.GetShape(), *Parameters);
 			Parameters->Data = GraphBuilder.CreateSRV(FRDGBufferSRVDesc(Data.GetBuffer(), PF_R32_FLOAT));
 			// NOTE: Indices tensor can be int64, but UE lacks support of int64 buffers. Here we use a 32-bit pixel format and 
 			// in the shader we discard the upper 32-bits of each value. This means that index values need to be representable by 32 bits
diff --git a/Engine/Plugins/Experimental/NNERuntimeRDG/Source/NNERuntimeRDG/Private/Hlsl/NNERuntimeRDGGatherElements.cpp b/Engine/Plugins/Experimental/NNERuntimeRDG/Source/NNERuntimeRDG/Private/Hlsl/NNERuntimeRDGGatherElements.cpp
--- a/Engine/Plugins/Experimental/NNERuntimeRDG/Source/NNERuntimeRDG/Private/Hlsl/NNERuntimeRDGGatherElements.cpp
+++ b/Engine/Plugins/Experimental/NNERuntimeRDG/Source/NNERuntimeRDG/Private/Hlsl/NNERuntimeRDGGatherElements.cpp
@@ -6,7 +6,7 @@
 #include "NNEHlslShadersGatherElementsCS.h"
 #include "NNEHlslShadersLog.h"
 #include "NNERuntimeRDGHlslHelper.h"
-#include "NNETensor.h"
+#include "NNERuntimeRDGTensor.h"
 #include "NNETypes.h"
 
 namespace UE::NNERuntimeRDG::Private::Hlsl
@@ -28,7 +28,7 @@ namespace UE::NNERuntimeRDG::Private::Hlsl
 		int32 Axis;
 	public:
 
-		virtual int PrepareOutputs(TConstArrayView<NNE::Internal::FTensorRef> InputTensors, TArrayView<NNE::Internal::FTensorRef> OutputTensors) override
+		virtual int PrepareOutputs(TConstArrayView<FTensorRef> InputTensors, TArrayView<FTensorRef> OutputTensors) override
 		{
 			check(InputTensors.Num() == 2);
 			check(OutputTensors.Num() == 1);
diff --git a/Engine/Plugins/Experimental/NNERuntimeRDG/Source/NNERuntimeRDG/Private/Hlsl/NNERuntimeRDGGemm.cpp b/Engine/Plugins/Experimental/NNERuntimeRDG/Source/NNERuntimeRDG/Private/Hlsl/NNERuntimeRDGGemm.cpp
--- a/Engine/Plugins/Experimental/NNERuntimeRDG/Source/NNERuntimeRDG/Private/Hlsl/NNERuntimeRDGGemm.cpp
+++ b/Engine/Plugins/Experimental/NNERuntimeRDG/Source/NNERuntimeRDG/Private/Hlsl/NNERuntimeRDGGemm.cpp
@@ -6,7 +6,7 @@
 #include "NNEHlslShadersLog.h"
 #include "NNERuntimeRDGDataAttributeMap.h"
 #include "NNERuntimeRDGHlslHelper.h"
-#include "NNETensor.h"
+#include "NNERuntimeRDGTensor.h"
 #include "NNETypes.h"
 
 namespace UE::NNERuntimeRDG::Private::Hlsl
@@ -32,7 +32,7 @@ namespace UE::NNERuntimeRDG::Private::Hlsl
 
 	public:
 
-		virtual int PrepareOutputs(TConstArrayView<NNE::Internal::FTensorRef> InputTensors, TArrayView<NNE::Internal::FTensorRef> OutputTensors) override
+		virtual int PrepareOutputs(TConstArrayView<FTensorRef> InputTensors, TArrayView<FTensorRef> OutputTensors) override
 		{
 			check(InputTensors.Num() >= 2 && InputTensors.Num() <= 3);
 			check(OutputTensors.Num() == 1);
@@ -132,8 +132,9 @@ namespace UE::NNERuntimeRDG::Private::Hlsl
 			check(CScalarMode != EGemmCScalar::MAX);
 			
 			// Set parameters
+			const NNE::FTensorShape* InputCShape = InputC != nullptr ? &InputC->GetShape() : nullptr;
 			TGemmCS::FParameters* Parameters = GraphBuilder.AllocParameters<TGemmCS::FParameters>();
-			TGemmCS::FillInParameters(InputAlpha, InputBeta, InputTransA, InputTransB, InputA, InputB, InputC, CConstantScalar, *Parameters);
+			TGemmCS::FillInParameters(InputAlpha, InputBeta, InputTransA, InputTransB, InputA.GetShape(), InputB.GetShape(), InputCShape, CConstantScalar, *Parameters);
 			Parameters->A = GraphBuilder.CreateSRV(FRDGBufferSRVDesc(InputA.GetBuffer(), PF_R32_FLOAT));
 			Parameters->B = GraphBuilder.CreateSRV(FRDGBufferSRVDesc(InputB.GetBuffer(), PF_R32_FLOAT));
 			if (InputC != nullptr) {
diff --git a/Engine/Plugins/Experimental/NNERuntimeRDG/Source/NNERuntimeRDG/Private/Hlsl/NNERuntimeRDGGlobalPool.cpp b/Engine/Plugins/Experimental/NNERuntimeRDG/Source/NNERuntimeRDG/Private/Hlsl/NNERuntimeRDGGlobalPool.cpp
--- a/Engine/Plugins/Experimental/NNERuntimeRDG/Source/NNERuntimeRDG/Private/Hlsl/NNERuntimeRDGGlobalPool.cpp
+++ b/Engine/Plugins/Experimental/NNERuntimeRDG/Source/NNERuntimeRDG/Private/Hlsl/NNERuntimeRDGGlobalPool.cpp
@@ -5,7 +5,7 @@
 #include "NNEHlslShadersLog.h"
 #include "NNEHlslShadersReduceCS.h"
 #include "NNERuntimeRDGHlslHelper.h"
-#include "NNETensor.h"
+#include "NNERuntimeRDGTensor.h"
 #include "NNETypes.h"
 #include "RenderGraphUtils.h"
 
@@ -31,7 +31,7 @@ namespace UE::NNERuntimeRDG::Private::Hlsl
 
 	public:
 
-		virtual int PrepareOutputs(TConstArrayView<NNE::Internal::FTensorRef> InputTensors, TArrayView<NNE::Internal::FTensorRef> OutputTensors) override
+		virtual int PrepareOutputs(TConstArrayView<FTensorRef> InputTensors, TArrayView<FTensorRef> OutputTensors) override
 		{
 			check(InputTensors.Num() == 1);
 			check(OutputTensors.Num() == 1);
diff --git a/Engine/Plugins/Experimental/NNERuntimeRDG/Source/NNERuntimeRDG/Private/Hlsl/NNERuntimeRDGIdentity.cpp b/Engine/Plugins/Experimental/NNERuntimeRDG/Source/NNERuntimeRDG/Private/Hlsl/NNERuntimeRDGIdentity.cpp
--- a/Engine/Plugins/Experimental/NNERuntimeRDG/Source/NNERuntimeRDG/Private/Hlsl/NNERuntimeRDGIdentity.cpp
+++ b/Engine/Plugins/Experimental/NNERuntimeRDG/Source/NNERuntimeRDG/Private/Hlsl/NNERuntimeRDGIdentity.cpp
@@ -1,7 +1,7 @@
 // Copyright Epic Games, Inc. All Rights Reserved.
 
 #include "NNERuntimeRDGIdentity.h"
-#include "NNETensor.h"
+#include "NNERuntimeRDGTensor.h"
 #include "NNETypes.h"
 #include "RenderGraphUtils.h"
 
@@ -21,12 +21,12 @@ namespace UE::NNERuntimeRDG::Private::Hlsl
 
 	public:
 
-		virtual int PrepareOutputs(TConstArrayView<NNE::Internal::FTensorRef> InputTensors, TArrayView<NNE::Internal::FTensorRef> OutputTensors) override
+		virtual int PrepareOutputs(TConstArrayView<FTensorRef> InputTensors, TArrayView<FTensorRef> OutputTensors) override
 		{
 			check(InputTensors.Num() == 1);
 			check(OutputTensors.Num() == 1);
 
-			const NNE::Internal::FTensor& X = *InputTensors[0];
+			const FTensor& X = *InputTensors[0];
 
 			OutputTensors[0]->SetShape(X.GetShape());
 			if (X.HasPreparedData())
diff --git a/Engine/Plugins/Experimental/NNERuntimeRDG/Source/NNERuntimeRDG/Private/Hlsl/NNERuntimeRDGInstanceNormalization.cpp b/Engine/Plugins/Experimental/NNERuntimeRDG/Source/NNERuntimeRDG/Private/Hlsl/NNERuntimeRDGInstanceNormalization.cpp
--- a/Engine/Plugins/Experimental/NNERuntimeRDG/Source/NNERuntimeRDG/Private/Hlsl/NNERuntimeRDGInstanceNormalization.cpp
+++ b/Engine/Plugins/Experimental/NNERuntimeRDG/Source/NNERuntimeRDG/Private/Hlsl/NNERuntimeRDGInstanceNormalization.cpp
@@ -6,7 +6,7 @@
 #include "NNEHlslShadersLog.h"
 #include "NNEHlslShadersReduceCS.h"
 #include "NNERuntimeRDGHlslHelper.h"
-#include "NNETensor.h"
+#include "NNERuntimeRDGTensor.h"
 #include "NNETypes.h"
 
 namespace
@@ -61,7 +61,7 @@ namespace UE::NNERuntimeRDG::Private::Hlsl
 
 	public:
 
-		virtual int PrepareOutputs(TConstArrayView<NNE::Internal::FTensorRef> InputTensors, TArrayView<NNE::Internal::FTensorRef> OutputTensors) override
+		virtual int PrepareOutputs(TConstArrayView<FTensorRef> InputTensors, TArrayView<FTensorRef> OutputTensors) override
 		{
 			check(InputTensors.Num() == 3);
 			check(OutputTensors.Num() == 1);
diff --git a/Engine/Plugins/Experimental/NNERuntimeRDG/Source/NNERuntimeRDG/Private/Hlsl/NNERuntimeRDGLayerNormalization.cpp b/Engine/Plugins/Experimental/NNERuntimeRDG/Source/NNERuntimeRDG/Private/Hlsl/NNERuntimeRDGLayerNormalization.cpp
--- a/Engine/Plugins/Experimental/NNERuntimeRDG/Source/NNERuntimeRDG/Private/Hlsl/NNERuntimeRDGLayerNormalization.cpp
+++ b/Engine/Plugins/Experimental/NNERuntimeRDG/Source/NNERuntimeRDG/Private/Hlsl/NNERuntimeRDGLayerNormalization.cpp
@@ -6,7 +6,7 @@
 #include "NNEHlslShadersLog.h"
 #include "NNEHlslShadersReduceCS.h"
 #include "NNERuntimeRDGHlslHelper.h"
-#include "NNETensor.h"
+#include "NNERuntimeRDGTensor.h"
 #include "NNETypes.h"
 
 namespace UE::NNERuntimeRDG::Private::Hlsl
@@ -31,7 +31,7 @@ namespace UE::NNERuntimeRDG::Private::Hlsl
 
 	public:
 
-		virtual int PrepareOutputs(TConstArrayView<NNE::Internal::FTensorRef> InputTensors, TArrayView<NNE::Internal::FTensorRef> OutputTensors) override
+		virtual int PrepareOutputs(TConstArrayView<FTensorRef> InputTensors, TArrayView<FTensorRef> OutputTensors) override
 		{
 			check(InputTensors.Num() >= 2 && InputTensors.Num() <= 3);
 			check(OutputTensors.Num() >= 1 && OutputTensors.Num() <= 3);
diff --git a/Engine/Plugins/Experimental/NNERuntimeRDG/Source/NNERuntimeRDG/Private/Hlsl/NNERuntimeRDGMatMul.cpp b/Engine/Plugins/Experimental/NNERuntimeRDG/Source/NNERuntimeRDG/Private/Hlsl/NNERuntimeRDGMatMul.cpp
--- a/Engine/Plugins/Experimental/NNERuntimeRDG/Source/NNERuntimeRDG/Private/Hlsl/NNERuntimeRDGMatMul.cpp
+++ b/Engine/Plugins/Experimental/NNERuntimeRDG/Source/NNERuntimeRDG/Private/Hlsl/NNERuntimeRDGMatMul.cpp
@@ -5,7 +5,7 @@
 #include "NNEHlslShadersGemmCS.h"
 #include "NNEHlslShadersLog.h"
 #include "NNERuntimeRDGHlslHelper.h"
-#include "NNETensor.h"
+#include "NNERuntimeRDGTensor.h"
 #include "NNETypes.h"
 
 namespace UE::NNERuntimeRDG::Private::Hlsl
@@ -24,7 +24,7 @@ namespace UE::NNERuntimeRDG::Private::Hlsl
 
 	public:
 
-		virtual int PrepareOutputs(TConstArrayView<NNE::Internal::FTensorRef> InputTensors, TArrayView<NNE::Internal::FTensorRef> OutputTensors) override
+		virtual int PrepareOutputs(TConstArrayView<FTensorRef> InputTensors, TArrayView<FTensorRef> OutputTensors) override
 		{
 			check(InputTensors.Num() == 2);
 			check(OutputTensors.Num() == 1);
@@ -115,7 +115,7 @@ namespace UE::NNERuntimeRDG::Private::Hlsl
 
 			// Set parameters
 			TGemmCS::FParameters* Parameters = GraphBuilder.AllocParameters<TGemmCS::FParameters>();
-			TGemmCS::FillInParametersMatMul(InputA, InputB, *Parameters);
+			TGemmCS::FillInParametersMatMul(InputA.GetShape(), InputB.GetShape(), *Parameters);
 			Parameters->A = GraphBuilder.CreateSRV(FRDGBufferSRVDesc(InputA.GetBuffer(), PF_R32_FLOAT));
 			Parameters->B = GraphBuilder.CreateSRV(FRDGBufferSRVDesc(InputB.GetBuffer(), PF_R32_FLOAT));
 			Parameters->Y = GraphBuilder.CreateUAV(FRDGBufferUAVDesc(Output.GetBuffer(), PF_R32_FLOAT));
diff --git a/Engine/Plugins/Experimental/NNERuntimeRDG/Source/NNERuntimeRDG/Private/Hlsl/NNERuntimeRDGPad.cpp b/Engine/Plugins/Experimental/NNERuntimeRDG/Source/NNERuntimeRDG/Private/Hlsl/NNERuntimeRDGPad.cpp
--- a/Engine/Plugins/Experimental/NNERuntimeRDG/Source/NNERuntimeRDG/Private/Hlsl/NNERuntimeRDGPad.cpp
+++ b/Engine/Plugins/Experimental/NNERuntimeRDG/Source/NNERuntimeRDG/Private/Hlsl/NNERuntimeRDGPad.cpp
@@ -6,7 +6,7 @@
 #include "NNEHlslShadersPadCS.h"
 #include "NNERuntimeRDGHlslHelper.h"
 #include "Helper/NNERuntimeRDGOperatorHelper.h"
-#include "NNETensor.h"
+#include "NNERuntimeRDGTensor.h"
 #include "NNETypes.h"
 
 namespace UE::NNERuntimeRDG::Private::Hlsl
@@ -32,17 +32,17 @@ namespace UE::NNERuntimeRDG::Private::Hlsl
 
 	public:
 
-		virtual int PrepareOutputs(TConstArrayView<NNE::Internal::FTensorRef> InputTensors, TArrayView<NNE::Internal::FTensorRef> OutputTensors) override
+		virtual int PrepareOutputs(TConstArrayView<FTensorRef> InputTensors, TArrayView<FTensorRef> OutputTensors) override
 		{
 			CheckInputTensorCount(InputTensors.Num());
 			check(OutputTensors.Num() == 1);
 
-			const NNE::Internal::FTensor& X = *InputTensors[0];
+			const FTensor& X = *InputTensors[0];
 			const int32 Rank = X.GetShape().Rank();
 
 			if constexpr (Version >= 11)
 			{
-				const NNE::Internal::FTensorRef PadsTensor = InputTensors[1];
+				const FTensorRef PadsTensor = InputTensors[1];
 				if(!PadsTensor->IsConstant())
 				{
 					UE_LOG(LogNNERuntimeRDGHlsl, Warning, TEXT("Pad: Input 'pads' (name: %s) should be constant."), *PadsTensor->GetName());
@@ -51,7 +51,7 @@ namespace UE::NNERuntimeRDG::Private::Hlsl
 
 				if (InputTensors.Num() >= 3 && !InputTensors[2]->IsEmpty())
 				{
-					const NNE::Internal::FTensor& ValueTensor = *InputTensors[2];
+					const FTensor& ValueTensor = *InputTensors[2];
 
 					if(!ValueTensor.HasPreparedData())
 					{
@@ -70,7 +70,7 @@ namespace UE::NNERuntimeRDG::Private::Hlsl
 
 				if (InputTensors.Num() >= 4 && !InputTensors[3]->IsEmpty())
 				{
-					const NNE::Internal::FTensorRef AxesTensor = InputTensors[3];
+					const FTensorRef AxesTensor = InputTensors[3];
 
 					if(!AxesTensor->IsConstant())
 					{
diff --git a/Engine/Plugins/Experimental/NNERuntimeRDG/Source/NNERuntimeRDG/Private/Hlsl/NNERuntimeRDGPool.cpp b/Engine/Plugins/Experimental/NNERuntimeRDG/Source/NNERuntimeRDG/Private/Hlsl/NNERuntimeRDGPool.cpp
--- a/Engine/Plugins/Experimental/NNERuntimeRDG/Source/NNERuntimeRDG/Private/Hlsl/NNERuntimeRDGPool.cpp
+++ b/Engine/Plugins/Experimental/NNERuntimeRDG/Source/NNERuntimeRDG/Private/Hlsl/NNERuntimeRDGPool.cpp
@@ -7,7 +7,7 @@
 #include "NNEHlslShadersPoolCS.h"
 #include "NNEHlslShadersTypeHelper.h"
 #include "NNERuntimeRDGHlslHelper.h"
-#include "NNETensor.h"
+#include "NNERuntimeRDGTensor.h"
 #include "NNETypes.h"
 #include "RenderGraphUtils.h"
 
@@ -44,12 +44,12 @@ namespace UE::NNERuntimeRDG::Private::Hlsl
 
 	public:
 
-		virtual int PrepareOutputs(TConstArrayView<NNE::Internal::FTensorRef> InputTensors, TArrayView<NNE::Internal::FTensorRef> OutputTensors) override
+		virtual int PrepareOutputs(TConstArrayView<FTensorRef> InputTensors, TArrayView<FTensorRef> OutputTensors) override
 		{
 			check(InputTensors.Num() == 1);
 			check(OutputTensors.Num() == 1);
 
-			const NNE::Internal::FTensor& X = *InputTensors[0];
+			const FTensor& X = *InputTensors[0];
 			TConstArrayView<uint32> InputShape = X.GetShape().GetData();
 			TArray<uint32, TInlineAllocator<NNE::FTensorShape::MaxRank>> OutputShape;
 
diff --git a/Engine/Plugins/Experimental/NNERuntimeRDG/Source/NNERuntimeRDG/Private/Hlsl/NNERuntimeRDGReduce.cpp b/Engine/Plugins/Experimental/NNERuntimeRDG/Source/NNERuntimeRDG/Private/Hlsl/NNERuntimeRDGReduce.cpp
--- a/Engine/Plugins/Experimental/NNERuntimeRDG/Source/NNERuntimeRDG/Private/Hlsl/NNERuntimeRDGReduce.cpp
+++ b/Engine/Plugins/Experimental/NNERuntimeRDG/Source/NNERuntimeRDG/Private/Hlsl/NNERuntimeRDGReduce.cpp
@@ -8,7 +8,7 @@
 #include "NNEHlslShadersLog.h"
 #include "NNEHlslShadersReduceCS.h"
 #include "NNERuntimeRDGHlslHelper.h"
-#include "NNETensor.h"
+#include "NNERuntimeRDGTensor.h"
 #include "NNETypes.h"
 #include "RenderGraphUtils.h"
 
@@ -40,7 +40,7 @@ namespace UE::NNERuntimeRDG::Private::Hlsl
 
 	public:
 
-		virtual int PrepareOutputs(TConstArrayView<NNE::Internal::FTensorRef> InputTensors, TArrayView<NNE::Internal::FTensorRef> OutputTensors) override
+		virtual int PrepareOutputs(TConstArrayView<FTensorRef> InputTensors, TArrayView<FTensorRef> OutputTensors) override
 		{
 			check(bAxesAsInput ? (InputTensors.Num() >= 1 && InputTensors.Num() <= 2) : InputTensors.Num() == 1);
 			check(OutputTensors.Num() == 1);
@@ -49,7 +49,7 @@ namespace UE::NNERuntimeRDG::Private::Hlsl
 			{
 				if(InputTensors.Num() == 2)
 				{
-					const NNE::Internal::FTensorRef AxesTensor = InputTensors[1];
+					const FTensorRef AxesTensor = InputTensors[1];
 					if(!AxesTensor->HasPreparedData())
 					{
 						UE_LOG(LogNNERuntimeRDGHlsl, Warning, TEXT("Reduce: Tensor `axes` (name: %s) must be CPU constant."), *InputTensors[1]->GetName());
diff --git a/Engine/Plugins/Experimental/NNERuntimeRDG/Source/NNERuntimeRDG/Private/Hlsl/NNERuntimeRDGReshape.cpp b/Engine/Plugins/Experimental/NNERuntimeRDG/Source/NNERuntimeRDG/Private/Hlsl/NNERuntimeRDGReshape.cpp
--- a/Engine/Plugins/Experimental/NNERuntimeRDG/Source/NNERuntimeRDG/Private/Hlsl/NNERuntimeRDGReshape.cpp
+++ b/Engine/Plugins/Experimental/NNERuntimeRDG/Source/NNERuntimeRDG/Private/Hlsl/NNERuntimeRDGReshape.cpp
@@ -4,7 +4,7 @@
 
 #include "Helper/NNERuntimeRDGHelperReshape.h"
 #include "NNEHlslShadersLog.h"
-#include "NNETensor.h"
+#include "NNERuntimeRDGTensor.h"
 #include "NNETypes.h"
 #include "RenderGraphUtils.h"
 
@@ -26,13 +26,13 @@ namespace UE::NNERuntimeRDG::Private::Hlsl
 
 	public:
 
-		virtual int PrepareOutputs(TConstArrayView<NNE::Internal::FTensorRef> InputTensors, TArrayView<NNE::Internal::FTensorRef> OutputTensors) override
+		virtual int PrepareOutputs(TConstArrayView<FTensorRef> InputTensors, TArrayView<FTensorRef> OutputTensors) override
 		{
 			check(InputTensors.Num() == 2);
 			check(OutputTensors.Num() == 1);
 
-			const NNE::Internal::FTensor& X = *InputTensors[0];
-			const NNE::Internal::FTensor& ShapeTensor = *InputTensors[1];
+			const FTensor& X = *InputTensors[0];
+			const FTensor& ShapeTensor = *InputTensors[1];
 			TArray<uint32> OutputShapeData;
 
 			check(ShapeTensor.GetDataType() == ENNETensorDataType::Int64);
diff --git a/Engine/Plugins/Experimental/NNERuntimeRDG/Source/NNERuntimeRDG/Private/Hlsl/NNERuntimeRDGResize.cpp b/Engine/Plugins/Experimental/NNERuntimeRDG/Source/NNERuntimeRDG/Private/Hlsl/NNERuntimeRDGResize.cpp
--- a/Engine/Plugins/Experimental/NNERuntimeRDG/Source/NNERuntimeRDG/Private/Hlsl/NNERuntimeRDGResize.cpp
+++ b/Engine/Plugins/Experimental/NNERuntimeRDG/Source/NNERuntimeRDG/Private/Hlsl/NNERuntimeRDGResize.cpp
@@ -8,7 +8,7 @@
 #include "NNEHlslShadersResizeCS.h"
 #include "NNEHlslShadersTypeHelper.h"
 #include "NNERuntimeRDGHlslHelper.h"
-#include "NNETensor.h"
+#include "NNERuntimeRDGTensor.h"
 #include "NNETypes.h"
 #include "RenderGraphUtils.h"
 
@@ -62,13 +62,13 @@ namespace UE::NNERuntimeRDG::Private::Hlsl
 
 	public:
 
-		virtual int PrepareOutputs(TConstArrayView<NNE::Internal::FTensorRef> InputTensors, TArrayView<NNE::Internal::FTensorRef> OutputTensors) override
+		virtual int PrepareOutputs(TConstArrayView<FTensorRef> InputTensors, TArrayView<FTensorRef> OutputTensors) override
 		{
 			check(InputTensors.Num() >= 3 && InputTensors.Num() <= 4);
 			check(OutputTensors.Num() == 1);
 
-			const NNE::Internal::FTensor& Input = *InputTensors[0];
-			const NNE::Internal::FTensor& Roi = *InputTensors[1];
+			const FTensor& Input = *InputTensors[0];
+			const FTensor& Roi = *InputTensors[1];
 
 			if(CoordTransMode == UE::NNEHlslShaders::Internal::ECoordTransMode::TfCropAndResize)
 			{
@@ -94,7 +94,7 @@ namespace UE::NNERuntimeRDG::Private::Hlsl
 		
 			if(InputTensors.Num() == 3)
 			{
-				const NNE::Internal::FTensor& Scales = *InputTensors[2];
+				const FTensor& Scales = *InputTensors[2];
 				if(!Scales.HasPreparedData())
 				{
 					UE_LOG(LogNNERuntimeRDGHlsl, Warning, TEXT("Resize: `scales` tensor could not be made constant. (name %s)."), *Scales.GetName());
@@ -126,7 +126,7 @@ namespace UE::NNERuntimeRDG::Private::Hlsl
 			
 			if(InputTensors.Num() == 4)
 			{
-				const NNE::Internal::FTensor& Sizes = *InputTensors[3];
+				const FTensor& Sizes = *InputTensors[3];
 				if(!Sizes.HasPreparedData())
 				{
 					UE_LOG(LogNNERuntimeRDGHlsl, Warning, TEXT("Resize: `sizes` tensor could not be made constant. (name %s)."), *Sizes.GetName());
diff --git a/Engine/Plugins/Experimental/NNERuntimeRDG/Source/NNERuntimeRDG/Private/Hlsl/NNERuntimeRDGScatterND.cpp b/Engine/Plugins/Experimental/NNERuntimeRDG/Source/NNERuntimeRDG/Private/Hlsl/NNERuntimeRDGScatterND.cpp
--- a/Engine/Plugins/Experimental/NNERuntimeRDG/Source/NNERuntimeRDG/Private/Hlsl/NNERuntimeRDGScatterND.cpp
+++ b/Engine/Plugins/Experimental/NNERuntimeRDG/Source/NNERuntimeRDG/Private/Hlsl/NNERuntimeRDGScatterND.cpp
@@ -6,7 +6,7 @@
 #include "NNEHlslShadersLog.h"
 #include "NNEHlslShadersScatterNDCS.h"
 #include "NNERuntimeRDGHlslHelper.h"
-#include "NNETensor.h"
+#include "NNERuntimeRDGTensor.h"
 #include "NNETypes.h"
 #include "RenderGraphUtils.h"
 
@@ -29,16 +29,16 @@ namespace UE::NNERuntimeRDG::Private::Hlsl
 
 	public:
 
-		virtual int PrepareOutputs(TConstArrayView<NNE::Internal::FTensorRef> InputTensors, TArrayView<NNE::Internal::FTensorRef> OutputTensors) override
+		virtual int PrepareOutputs(TConstArrayView<FTensorRef> InputTensors, TArrayView<FTensorRef> OutputTensors) override
 		{
 			check(InputTensors.Num() == 3);
 			check(OutputTensors.Num() == 1);
 
-			const NNE::Internal::FTensor& Input = *InputTensors[0];
+			const FTensor& Input = *InputTensors[0];
 			TConstArrayView<uint32> InputShape = Input.GetShape().GetData();
-			const NNE::Internal::FTensor& Indices = *InputTensors[1];
+			const FTensor& Indices = *InputTensors[1];
 			TConstArrayView<uint32> IndicesShape = Indices.GetShape().GetData();
-			const NNE::Internal::FTensor& Updates = *InputTensors[2];
+			const FTensor& Updates = *InputTensors[2];
 			TConstArrayView<uint32> UpdatesShape = Updates.GetShape().GetData();
 
 			if (IndicesShape.Last() > (uint32) InputShape.Num())
diff --git a/Engine/Plugins/Experimental/NNERuntimeRDG/Source/NNERuntimeRDG/Private/Hlsl/NNERuntimeRDGShape.cpp b/Engine/Plugins/Experimental/NNERuntimeRDG/Source/NNERuntimeRDG/Private/Hlsl/NNERuntimeRDGShape.cpp
--- a/Engine/Plugins/Experimental/NNERuntimeRDG/Source/NNERuntimeRDG/Private/Hlsl/NNERuntimeRDGShape.cpp
+++ b/Engine/Plugins/Experimental/NNERuntimeRDG/Source/NNERuntimeRDG/Private/Hlsl/NNERuntimeRDGShape.cpp
@@ -3,7 +3,7 @@
 #include "NNERuntimeRDGShape.h"
 
 #include "NNEHlslShadersLog.h"
-#include "NNETensor.h"
+#include "NNERuntimeRDGTensor.h"
 #include "NNETypes.h"
 
 namespace UE::NNERuntimeRDG::Private::Hlsl
@@ -20,13 +20,13 @@ namespace UE::NNERuntimeRDG::Private::Hlsl
 
 	public:
 
-		virtual int PrepareOutputs(TConstArrayView<NNE::Internal::FTensorRef> InputTensors, TArrayView<NNE::Internal::FTensorRef> OutputTensors) override
+		virtual int PrepareOutputs(TConstArrayView<FTensorRef> InputTensors, TArrayView<FTensorRef> OutputTensors) override
 		{
 			check(InputTensors.Num() == 1);
 			check(OutputTensors.Num() == 1);
 			check(OutputTensors[0]->GetDataType() == ENNETensorDataType::Int64);
 
-			const NNE::Internal::FTensor& X = *InputTensors[0];
+			const FTensor& X = *InputTensors[0];
 
 			TArray<uint32> OutputShapeData;
 			
diff --git a/Engine/Plugins/Experimental/NNERuntimeRDG/Source/NNERuntimeRDG/Private/Hlsl/NNERuntimeRDGSize.cpp b/Engine/Plugins/Experimental/NNERuntimeRDG/Source/NNERuntimeRDG/Private/Hlsl/NNERuntimeRDGSize.cpp
--- a/Engine/Plugins/Experimental/NNERuntimeRDG/Source/NNERuntimeRDG/Private/Hlsl/NNERuntimeRDGSize.cpp
+++ b/Engine/Plugins/Experimental/NNERuntimeRDG/Source/NNERuntimeRDG/Private/Hlsl/NNERuntimeRDGSize.cpp
@@ -3,7 +3,7 @@
 #include "NNERuntimeRDGSize.h"
 
 #include "NNEHlslShadersLog.h"
-#include "NNETensor.h"
+#include "NNERuntimeRDGTensor.h"
 #include "NNETypes.h"
 
 namespace UE::NNERuntimeRDG::Private::Hlsl
@@ -22,12 +22,12 @@ namespace UE::NNERuntimeRDG::Private::Hlsl
 
 	public:
 
-		virtual int PrepareOutputs(TConstArrayView<NNE::Internal::FTensorRef> InputTensors, TArrayView<NNE::Internal::FTensorRef> OutputTensors) override
+		virtual int PrepareOutputs(TConstArrayView<FTensorRef> InputTensors, TArrayView<FTensorRef> OutputTensors) override
 		{
 			check(InputTensors.Num() == 1);
 			check(OutputTensors.Num() == 1);
 
-			const NNE::Internal::FTensor& X = *InputTensors[0];
+			const FTensor& X = *InputTensors[0];
 			NNE::FTensorShape ScalarShape = NNE::FTensorShape::Make({});
 			TArray<int64> OutputData;
 			
diff --git a/Engine/Plugins/Experimental/NNERuntimeRDG/Source/NNERuntimeRDG/Private/Hlsl/NNERuntimeRDGSlice.cpp b/Engine/Plugins/Experimental/NNERuntimeRDG/Source/NNERuntimeRDG/Private/Hlsl/NNERuntimeRDGSlice.cpp
--- a/Engine/Plugins/Experimental/NNERuntimeRDG/Source/NNERuntimeRDG/Private/Hlsl/NNERuntimeRDGSlice.cpp
+++ b/Engine/Plugins/Experimental/NNERuntimeRDG/Source/NNERuntimeRDG/Private/Hlsl/NNERuntimeRDGSlice.cpp
@@ -2,12 +2,12 @@
 
 #include "NNERuntimeRDGSlice.h"
 
+#include "Helper/NNERuntimeRDGHelperSlice.h"
 #include "Helper/NNERuntimeRDGOperatorHelper.h"
 #include "NNEHlslShadersLog.h"
 #include "NNEHlslShadersSliceCS.h"
-#include "NNERuntimeRDGHelperSlice.h"
 #include "NNERuntimeRDGHlslHelper.h"
-#include "NNETensor.h"
+#include "NNERuntimeRDGTensor.h"
 #include "NNETypes.h"
 #include "RenderGraphUtils.h"
 
@@ -37,7 +37,7 @@ namespace UE::NNERuntimeRDG::Private::Hlsl
 		TArray<int32, TInlineAllocator<NNE::FTensorShape::MaxRank>> End;
 		TArray<int32, TInlineAllocator<NNE::FTensorShape::MaxRank>> Step;
 
-		bool TryGetAttributesFromConstantTensors(TConstArrayView<NNE::Internal::FTensorRef> InputTensors)
+		bool TryGetAttributesFromConstantTensors(TConstArrayView<FTensorRef> InputTensors)
 		{
 			check(InputTensors.Num() >= 3 && InputTensors.Num() <= 5);
 
@@ -180,7 +180,7 @@ namespace UE::NNERuntimeRDG::Private::Hlsl
 
 	public:
 
-		virtual int PrepareOutputs(TConstArrayView<NNE::Internal::FTensorRef> InputTensors, TArrayView<NNE::Internal::FTensorRef> OutputTensors) override
+		virtual int PrepareOutputs(TConstArrayView<FTensorRef> InputTensors, TArrayView<FTensorRef> OutputTensors) override
 		{
 			check(OutputTensors.Num() == 1);
 
@@ -221,7 +221,7 @@ namespace UE::NNERuntimeRDG::Private::Hlsl
 
 			OutputTensors[0]->SetShape(OutputShape);
 			
-			Internal::CPUHelper::Slice::Apply(*InputTensors[0], *OutputTensors[0], Start, Step);
+			CPUHelper::Slice::Apply(*InputTensors[0], *OutputTensors[0], Start, Step);
 
 			if (InputTensors[0]->HasPreparedData() && !OutputTensors[0]->HasPreparedData())
 			{
diff --git a/Engine/Plugins/Experimental/NNERuntimeRDG/Source/NNERuntimeRDG/Private/Hlsl/NNERuntimeRDGSoftmax.cpp b/Engine/Plugins/Experimental/NNERuntimeRDG/Source/NNERuntimeRDG/Private/Hlsl/NNERuntimeRDGSoftmax.cpp
--- a/Engine/Plugins/Experimental/NNERuntimeRDG/Source/NNERuntimeRDG/Private/Hlsl/NNERuntimeRDGSoftmax.cpp
+++ b/Engine/Plugins/Experimental/NNERuntimeRDG/Source/NNERuntimeRDG/Private/Hlsl/NNERuntimeRDGSoftmax.cpp
@@ -6,7 +6,7 @@
 #include "NNEHlslShadersReduceCS.h"
 #include "NNEHlslShadersSoftmaxCS.h"
 #include "NNERuntimeRDGHlslHelper.h"
-#include "NNETensor.h"
+#include "NNERuntimeRDGTensor.h"
 #include "NNETypes.h"
 #include "RenderGraphUtils.h"
 
@@ -32,7 +32,7 @@ namespace UE::NNERuntimeRDG::Private::Hlsl
 
 	public:
 
-		virtual int PrepareOutputs(TConstArrayView<NNE::Internal::FTensorRef> InputTensors, TArrayView<NNE::Internal::FTensorRef> OutputTensors) override
+		virtual int PrepareOutputs(TConstArrayView<FTensorRef> InputTensors, TArrayView<FTensorRef> OutputTensors) override
 		{
 			check(InputTensors.Num() == 1);
 			check(OutputTensors.Num() == 1);
diff --git a/Engine/Plugins/Experimental/NNERuntimeRDG/Source/NNERuntimeRDG/Private/Hlsl/NNERuntimeRDGSplit.cpp b/Engine/Plugins/Experimental/NNERuntimeRDG/Source/NNERuntimeRDG/Private/Hlsl/NNERuntimeRDGSplit.cpp
--- a/Engine/Plugins/Experimental/NNERuntimeRDG/Source/NNERuntimeRDG/Private/Hlsl/NNERuntimeRDGSplit.cpp
+++ b/Engine/Plugins/Experimental/NNERuntimeRDG/Source/NNERuntimeRDG/Private/Hlsl/NNERuntimeRDGSplit.cpp
@@ -5,7 +5,7 @@
 #include "NNEHlslShadersLog.h"
 #include "NNEHlslShadersSplitCS.h"
 #include "NNERuntimeRDGHlslHelper.h"
-#include "NNETensor.h"
+#include "NNERuntimeRDGTensor.h"
 #include "NNETypes.h"
 #include "RenderGraphUtils.h"
 #include "Algo/Accumulate.h"
@@ -37,7 +37,7 @@ namespace UE::NNERuntimeRDG::Private::Hlsl
 
 	public:
 
-		virtual int PrepareOutputs(TConstArrayView<NNE::Internal::FTensorRef> InputTensors, TArrayView<NNE::Internal::FTensorRef> OutputTensors) override
+		virtual int PrepareOutputs(TConstArrayView<FTensorRef> InputTensors, TArrayView<FTensorRef> OutputTensors) override
 		{
 			using namespace UE::NNEHlslShaders::Internal;
 
@@ -48,7 +48,7 @@ namespace UE::NNERuntimeRDG::Private::Hlsl
 
 			if(Version >= 13 && InputTensors.Num() == 2)
 			{
-				const NNE::Internal::FTensor& SplitTensor = *InputTensors[1];
+				const FTensor& SplitTensor = *InputTensors[1];
 				if(!SplitTensor.HasPreparedData())
 				{
 					UE_LOG(LogNNERuntimeRDGHlsl, Warning, TEXT("Split: split tensor need to be CPU-constant in order to compute output shapes."));
diff --git a/Engine/Plugins/Experimental/NNERuntimeRDG/Source/NNERuntimeRDG/Private/Hlsl/NNERuntimeRDGSqueeze.cpp b/Engine/Plugins/Experimental/NNERuntimeRDG/Source/NNERuntimeRDG/Private/Hlsl/NNERuntimeRDGSqueeze.cpp
--- a/Engine/Plugins/Experimental/NNERuntimeRDG/Source/NNERuntimeRDG/Private/Hlsl/NNERuntimeRDGSqueeze.cpp
+++ b/Engine/Plugins/Experimental/NNERuntimeRDG/Source/NNERuntimeRDG/Private/Hlsl/NNERuntimeRDGSqueeze.cpp
@@ -3,7 +3,7 @@
 #include "NNERuntimeRDGSqueeze.h"
 
 #include "NNEHlslShadersLog.h"
-#include "NNETensor.h"
+#include "NNERuntimeRDGTensor.h"
 #include "NNETypes.h"
 #include "RenderGraphUtils.h"
 
@@ -25,12 +25,12 @@ namespace UE::NNERuntimeRDG::Private::Hlsl
 
 	public:
 
-		virtual int PrepareOutputs(TConstArrayView<NNE::Internal::FTensorRef> InputTensors, TArrayView<NNE::Internal::FTensorRef> OutputTensors) override
+		virtual int PrepareOutputs(TConstArrayView<FTensorRef> InputTensors, TArrayView<FTensorRef> OutputTensors) override
 		{
 			check(InputTensors.Num() == 1);
 			check(OutputTensors.Num() == 1);
 
-			const NNE::Internal::FTensor& X = *InputTensors[0];
+			const FTensor& X = *InputTensors[0];
 			TArray<uint32> OutputShapeData(X.GetShape().GetData());
 
 			for (int32 Idx = Axes.Num() - 1; Idx >= 0; --Idx)
diff --git a/Engine/Plugins/Experimental/NNERuntimeRDG/Source/NNERuntimeRDG/Private/Hlsl/NNERuntimeRDGTranspose.cpp b/Engine/Plugins/Experimental/NNERuntimeRDG/Source/NNERuntimeRDG/Private/Hlsl/NNERuntimeRDGTranspose.cpp
--- a/Engine/Plugins/Experimental/NNERuntimeRDG/Source/NNERuntimeRDG/Private/Hlsl/NNERuntimeRDGTranspose.cpp
+++ b/Engine/Plugins/Experimental/NNERuntimeRDG/Source/NNERuntimeRDG/Private/Hlsl/NNERuntimeRDGTranspose.cpp
@@ -2,11 +2,11 @@
 
 #include "NNERuntimeRDGTranspose.h"
 
+#include "Helper/NNERuntimeRDGHelperTranspose.h"
 #include "NNEHlslShadersLog.h"
 #include "NNEHlslShadersTransposeCS.h"
-#include "NNERuntimeRDGHelperTranspose.h"
 #include "NNERuntimeRDGHlslHelper.h"
-#include "NNETensor.h"
+#include "NNERuntimeRDGTensor.h"
 #include "NNETypes.h"
 #include "RenderGraphUtils.h"
 
@@ -29,12 +29,12 @@ namespace UE::NNERuntimeRDG::Private::Hlsl
 
 	public:
 
-		virtual int PrepareOutputs(TConstArrayView<NNE::Internal::FTensorRef> InputTensors, TArrayView<NNE::Internal::FTensorRef> OutputTensors) override
+		virtual int PrepareOutputs(TConstArrayView<FTensorRef> InputTensors, TArrayView<FTensorRef> OutputTensors) override
 		{
 			check(InputTensors.Num() == 1);
 			check(OutputTensors.Num() == 1);
 
-			const NNE::Internal::FTensor& X = *InputTensors[0];
+			const FTensor& X = *InputTensors[0];
 			TConstArrayView<uint32> InputShape = X.GetShape().GetData();
 			TArray<uint32, TInlineAllocator<NNE::FTensorShape::MaxRank>> OutputShape;
 
@@ -45,7 +45,7 @@ namespace UE::NNERuntimeRDG::Private::Hlsl
 			}
 			OutputTensors[0]->SetShape(NNE::FTensorShape::Make(OutputShape));
 
-			Internal::CPUHelper::Transpose::Apply(*InputTensors[0], Perm, *OutputTensors[0]);
+			CPUHelper::Transpose::Apply(*InputTensors[0], Perm, *OutputTensors[0]);
 
 			return 0;
 		};
diff --git a/Engine/Plugins/Experimental/NNERuntimeRDG/Source/NNERuntimeRDG/Private/Hlsl/NNERuntimeRDGUnsqueeze.cpp b/Engine/Plugins/Experimental/NNERuntimeRDG/Source/NNERuntimeRDG/Private/Hlsl/NNERuntimeRDGUnsqueeze.cpp
--- a/Engine/Plugins/Experimental/NNERuntimeRDG/Source/NNERuntimeRDG/Private/Hlsl/NNERuntimeRDGUnsqueeze.cpp
+++ b/Engine/Plugins/Experimental/NNERuntimeRDG/Source/NNERuntimeRDG/Private/Hlsl/NNERuntimeRDGUnsqueeze.cpp
@@ -3,7 +3,7 @@
 #include "NNERuntimeRDGUnsqueeze.h"
 
 #include "NNEHlslShadersLog.h"
-#include "NNETensor.h"
+#include "NNERuntimeRDGTensor.h"
 #include "NNETypes.h"
 #include "RenderGraphUtils.h"
 
@@ -25,12 +25,12 @@ namespace UE::NNERuntimeRDG::Private::Hlsl
 
 	public:
 
-		virtual int PrepareOutputs(TConstArrayView<NNE::Internal::FTensorRef> InputTensors, TArrayView<NNE::Internal::FTensorRef> OutputTensors) override
+		virtual int PrepareOutputs(TConstArrayView<FTensorRef> InputTensors, TArrayView<FTensorRef> OutputTensors) override
 		{
 			check(InputTensors.Num() == 1);
 			check(OutputTensors.Num() == 1);
 
-			const NNE::Internal::FTensor& X = *InputTensors[0];
+			const FTensor& X = *InputTensors[0];
 			TArray<uint32> OutputShapeData(X.GetShape().GetData());
 
 			for (int32 Axe : Axes)
diff --git a/Engine/Plugins/Experimental/NNERuntimeRDG/Source/NNERuntimeRDG/Private/Hlsl/NNERuntimeRDGUpsample.cpp b/Engine/Plugins/Experimental/NNERuntimeRDG/Source/NNERuntimeRDG/Private/Hlsl/NNERuntimeRDGUpsample.cpp
--- a/Engine/Plugins/Experimental/NNERuntimeRDG/Source/NNERuntimeRDG/Private/Hlsl/NNERuntimeRDGUpsample.cpp
+++ b/Engine/Plugins/Experimental/NNERuntimeRDG/Source/NNERuntimeRDG/Private/Hlsl/NNERuntimeRDGUpsample.cpp
@@ -6,7 +6,7 @@
 #include "NNERuntimeRDGAttributes.h"
 #include "NNEHlslShadersUpsampleCS.h"
 #include "NNERuntimeRDGHlslHelper.h"
-#include "NNETensor.h"
+#include "NNERuntimeRDGTensor.h"
 #include "NNETypes.h"
 
 namespace UE::NNERuntimeRDG::Private::Hlsl
@@ -29,13 +29,13 @@ namespace UE::NNERuntimeRDG::Private::Hlsl
 
 	public:
 
-		virtual int PrepareOutputs(TConstArrayView<NNE::Internal::FTensorRef> InputTensors, TArrayView<NNE::Internal::FTensorRef> OutputTensors) override
+		virtual int PrepareOutputs(TConstArrayView<FTensorRef> InputTensors, TArrayView<FTensorRef> OutputTensors) override
 		{
 			check(InputTensors.Num() == 2);
 			check(OutputTensors.Num() == 1);
 
-			const NNE::Internal::FTensor& Input = *InputTensors[0];
-			const NNE::Internal::FTensor& Scales = *InputTensors[1];
+			const FTensor& Input = *InputTensors[0];
+			const FTensor& Scales = *InputTensors[1];
 			
 			if (!Scales.HasPreparedData())
 			{
diff --git a/Engine/Plugins/Experimental/NNERuntimeRDG/Source/NNERuntimeRDG/Private/NNERuntimeRDGAttributes.h b/Engine/Plugins/Experimental/NNERuntimeRDG/Source/NNERuntimeRDG/Private/NNERuntimeRDGAttributes.h
new file mode 100644
--- /dev/null
+++ b/Engine/Plugins/Experimental/NNERuntimeRDG/Source/NNERuntimeRDG/Private/NNERuntimeRDGAttributes.h
@@ -0,0 +1,21 @@
+// Copyright Epic Games, Inc. All Rights Reserved.
+
+#pragma once
+
+#include "HAL/Platform.h"
+
+namespace UE::NNERuntimeRDG::Private
+{
+
+namespace AttrName
+{
+	static constexpr auto* Mode = TEXT("mode");
+}
+
+namespace AttrValue
+{
+	static constexpr auto* Nearest = TEXT("nearest");
+	static constexpr auto* Linear = TEXT("linear");
+}
+
+} // UE::NNERuntimeRDG::Private
\ No newline at end of file
diff --git a/Engine/Plugins/Experimental/NNERuntimeRDG/Source/NNERuntimeRDG/Private/NNERuntimeRDGBase.h b/Engine/Plugins/Experimental/NNERuntimeRDG/Source/NNERuntimeRDG/Private/NNERuntimeRDGBase.h
--- a/Engine/Plugins/Experimental/NNERuntimeRDG/Source/NNERuntimeRDG/Private/NNERuntimeRDGBase.h
+++ b/Engine/Plugins/Experimental/NNERuntimeRDG/Source/NNERuntimeRDG/Private/NNERuntimeRDGBase.h
@@ -11,7 +11,7 @@
 #ifdef NNE_UTILITIES_AVAILABLE
 #include "NNERuntimeRDGUtilsModelOptimizerInterface.h"
 #endif //NNE_UTILITIES_AVAILABLE
-#include "NNETensor.h"
+#include "NNERuntimeRDGTensor.h"
 #include "NNETypes.h"
 #include "RenderGraphResources.h"
 #include "Serialization/MemoryReader.h"
@@ -36,7 +36,7 @@ namespace UE::NNERuntimeRDG::Private
 struct IPrepareOperator
 {
 	virtual ~IPrepareOperator() = default;
-	virtual int PrepareOutputs(TConstArrayView<NNE::Internal::FTensorRef> InputTensors, TArrayView<NNE::Internal::FTensorRef> OutputTensors) = 0;
+	virtual int PrepareOutputs(TConstArrayView<FTensorRef> InputTensors, TArrayView<FTensorRef> OutputTensors) = 0;
 };
 
 /**
@@ -47,7 +47,7 @@ struct IOperatorRDG
 	virtual ~IOperatorRDG() = default;
 };
 
-class FTensorRDG : public NNE::Internal::FTensor
+class FTensorRDG : public FTensor
 {
 	FRDGBufferRef Buffer{};
 	
diff --git a/Engine/Plugins/Experimental/NNERuntimeRDG/Source/NNERuntimeRDG/Private/NNERuntimeRDGHlslHelper.cpp b/Engine/Plugins/Experimental/NNERuntimeRDG/Source/NNERuntimeRDG/Private/NNERuntimeRDGHlslHelper.cpp
--- a/Engine/Plugins/Experimental/NNERuntimeRDG/Source/NNERuntimeRDG/Private/NNERuntimeRDGHlslHelper.cpp
+++ b/Engine/Plugins/Experimental/NNERuntimeRDG/Source/NNERuntimeRDG/Private/NNERuntimeRDGHlslHelper.cpp
@@ -2,13 +2,13 @@
 
 #include "NNERuntimeRDGHlslHelper.h"
 #include "NNETypes.h"
-#include "NNETensor.h"
+#include "NNERuntimeRDGTensor.h"
 #include "RHI.h"
 
 namespace UE::NNERuntimeRDG::Private::Hlsl
 {
 	template<typename TTensorInfoParam>
-	void FillTensorSizeShaderParameters(const NNE::Internal::FTensor& Tensor, TTensorInfoParam& OutTensorInfoParam, int32 Idx)
+	void FillTensorSizeShaderParameters(const FTensor& Tensor, TTensorInfoParam& OutTensorInfoParam, int32 Idx)
 	{
 		static_assert(NNE::FTensorShape::MaxRank <= NXRT_TENSORSTRIDEINFO_MAX_NUM_DIMENSIONS);
 		for (int32 i = 0; i < Tensor.GetShape().Rank(); ++i)
@@ -16,11 +16,11 @@ namespace UE::NNERuntimeRDG::Private::Hlsl
 			OutTensorInfoParam[i][Idx] = Tensor.GetShape().GetData()[i];
 		}
 	}
-	template void FillTensorSizeShaderParameters(const NNE::Internal::FTensor&, FTensorInfoParam&, int32);
-	template void FillTensorSizeShaderParameters(const NNE::Internal::FTensor&, FTensorInfoParamArraySpan&, int32);
+	template void FillTensorSizeShaderParameters(const FTensor&, FTensorInfoParam&, int32);
+	template void FillTensorSizeShaderParameters(const FTensor&, FTensorInfoParamArraySpan&, int32);
 	
 	template<typename TTensorInfoParam>
-	void FillTensorStrideShaderParameters(const NNE::Internal::FTensor& Tensor, TTensorInfoParam& OutTensorInfoParam, int32 Idx, int32 TargetNumdimensionForBroadcast)
+	void FillTensorStrideShaderParameters(const FTensor& Tensor, TTensorInfoParam& OutTensorInfoParam, int32 Idx, int32 TargetNumdimensionForBroadcast)
 	{
 		if (TargetNumdimensionForBroadcast == -1)
 		{
@@ -46,11 +46,11 @@ namespace UE::NNERuntimeRDG::Private::Hlsl
 			}
 		}
 	}
-	template void FillTensorStrideShaderParameters(const NNE::Internal::FTensor&, FTensorInfoParam&, int32, int32);
-	template void FillTensorStrideShaderParameters(const NNE::Internal::FTensor&, FTensorInfoParamArraySpan&, int32, int32);
+	template void FillTensorStrideShaderParameters(const FTensor&, FTensorInfoParam&, int32, int32);
+	template void FillTensorStrideShaderParameters(const FTensor&, FTensorInfoParamArraySpan&, int32, int32);
 
 	template<typename TTensorInfoParam>
-	void FillTensorStrideForBroadcastShaderParameters(const NNE::Internal::FTensor& Tensor, int32 OutputNumdimension, TTensorInfoParam& OutTensorInfoParam, int32 Idx)
+	void FillTensorStrideForBroadcastShaderParameters(const FTensor& Tensor, int32 OutputNumdimension, TTensorInfoParam& OutTensorInfoParam, int32 Idx)
 	{
 		checkf(OutputNumdimension >= Tensor.GetShape().Rank(), TEXT("Can't broadcast tensor from rank %d to rank %d, should be inferior or equal."), Tensor.GetShape().Rank(), OutputNumdimension);
 		FillTensorStrideShaderParameters(Tensor, OutTensorInfoParam, Idx, OutputNumdimension);
@@ -64,8 +64,8 @@ namespace UE::NNERuntimeRDG::Private::Hlsl
 			}
 		}
 	}
-	template void FillTensorStrideForBroadcastShaderParameters(const NNE::Internal::FTensor&, int32, FTensorInfoParam&, int32);
-	template void FillTensorStrideForBroadcastShaderParameters(const NNE::Internal::FTensor&, int32, FTensorInfoParamArraySpan&, int32);
+	template void FillTensorStrideForBroadcastShaderParameters(const FTensor&, int32, FTensorInfoParam&, int32);
+	template void FillTensorStrideForBroadcastShaderParameters(const FTensor&, int32, FTensorInfoParamArraySpan&, int32);
 
 	FIntVector ComputeElementWiseThreadGroups(uint32 ElementCount, uint32 GroupSizeX)
 	{
diff --git a/Engine/Plugins/Experimental/NNERuntimeRDG/Source/NNERuntimeRDG/Private/NNERuntimeRDGHlslHelper.h b/Engine/Plugins/Experimental/NNERuntimeRDG/Source/NNERuntimeRDG/Private/NNERuntimeRDGHlslHelper.h
--- a/Engine/Plugins/Experimental/NNERuntimeRDG/Source/NNERuntimeRDG/Private/NNERuntimeRDGHlslHelper.h
+++ b/Engine/Plugins/Experimental/NNERuntimeRDG/Source/NNERuntimeRDG/Private/NNERuntimeRDGHlslHelper.h
@@ -5,7 +5,7 @@
 #include "Containers/StaticArray.h"
 #include "Templates/Function.h"
 
-namespace UE::NNE::Internal { class FTensor; }
+namespace UE::NNERuntimeRDG::Private { class FTensor; }
 
 #define NXRT_TENSORSTRIDEINFO_MAX_NUM_DIMENSIONS 8
 
@@ -27,10 +27,10 @@ namespace UE::NNERuntimeRDG::Private::Hlsl
 	};
 
 	template<typename TTensorInfoParam>
-	void FillTensorSizeShaderParameters(const NNE::Internal::FTensor& Tensor, TTensorInfoParam& OutTensorInfoParam, int32 Idx);
+	void FillTensorSizeShaderParameters(const FTensor& Tensor, TTensorInfoParam& OutTensorInfoParam, int32 Idx);
 	template<typename TTensorInfoParam>
-	void FillTensorStrideShaderParameters(const NNE::Internal::FTensor& Tensor, TTensorInfoParam& OutTensorInfoParam, int32 Idx, int32 TargetNumdimensionForBroadcast = -1);
+	void FillTensorStrideShaderParameters(const FTensor& Tensor, TTensorInfoParam& OutTensorInfoParam, int32 Idx, int32 TargetNumdimensionForBroadcast = -1);
 	template<typename TTensorInfoParam>
-	void FillTensorStrideForBroadcastShaderParameters(const NNE::Internal::FTensor& Tensor, int32 OutputNumdimension, TTensorInfoParam& OutTensorInfoParam, int32 Idx);
+	void FillTensorStrideForBroadcastShaderParameters(const FTensor& Tensor, int32 OutputNumdimension, TTensorInfoParam& OutTensorInfoParam, int32 Idx);
 	FIntVector ComputeElementWiseThreadGroups(uint32 ElementCount, uint32 GroupSizeX);
 } // namespace UE::NNERuntimeRDG::Private::Hlsl
\ No newline at end of file
diff --git a/Engine/Plugins/Experimental/NNERuntimeRDG/Source/NNERuntimeRDG/Private/NNERuntimeRDGModel.h b/Engine/Plugins/Experimental/NNERuntimeRDG/Source/NNERuntimeRDG/Private/NNERuntimeRDGModel.h
--- a/Engine/Plugins/Experimental/NNERuntimeRDG/Source/NNERuntimeRDG/Private/NNERuntimeRDGModel.h
+++ b/Engine/Plugins/Experimental/NNERuntimeRDG/Source/NNERuntimeRDG/Private/NNERuntimeRDGModel.h
@@ -4,7 +4,6 @@
 
 #include "CoreMinimal.h"
 #include "NNERuntimeRDG.h"
-#include "NNETensor.h"
 #include "NNETypes.h"
 #include "NNERuntimeRDGBase.h"
 
diff --git a/Engine/Plugins/Experimental/NNERuntimeRDG/Source/NNERuntimeRDG/Private/NNERuntimeRDGModelHlsl.cpp b/Engine/Plugins/Experimental/NNERuntimeRDG/Source/NNERuntimeRDG/Private/NNERuntimeRDGModelHlsl.cpp
--- a/Engine/Plugins/Experimental/NNERuntimeRDG/Source/NNERuntimeRDG/Private/NNERuntimeRDGModelHlsl.cpp
+++ b/Engine/Plugins/Experimental/NNERuntimeRDG/Source/NNERuntimeRDG/Private/NNERuntimeRDGModelHlsl.cpp
@@ -3,7 +3,7 @@
 #include "NNERuntimeRDGModelHlsl.h"
 
 #include "NNEHlslShadersLog.h"
-#include "NNETensor.h"
+#include "NNERuntimeRDGTensor.h"
 #include "NNERuntimeRDGDataAttributeMap.h"
 #include "NNERuntimeRDGHlsl.h"
 #include "NNERuntimeRDGHlslOp.h"
@@ -240,8 +240,8 @@ int FModelInstance::PrepareTensorShapesAndData()
 	// Run model preparation (including shape inference) on all operators
 	// This loop could be abstracted to a different runtime/system as it apply on FTensorRef & IPrepareOperator witch are RDG agnostics.
 	static constexpr int32 MaxExpectedInput = 10;
-	TArray<NNE::Internal::FTensorRef, TInlineAllocator<MaxExpectedInput>> InputTensors;
-	TArray<NNE::Internal::FTensorRef> OutputTensors;
+	TArray<FTensorRef, TInlineAllocator<MaxExpectedInput>> InputTensors;
+	TArray<FTensorRef> OutputTensors;
 	TArray<bool> AllInitializedTensors;
 
 	checkCode(
diff --git a/Engine/Plugins/Experimental/NNERuntimeRDG/Source/NNERuntimeRDG/Private/NNERuntimeRDGTensor.h b/Engine/Plugins/Experimental/NNERuntimeRDG/Source/NNERuntimeRDG/Private/NNERuntimeRDGTensor.h
new file mode 100644
--- /dev/null
+++ b/Engine/Plugins/Experimental/NNERuntimeRDG/Source/NNERuntimeRDG/Private/NNERuntimeRDGTensor.h
@@ -0,0 +1,123 @@
+// Copyright Epic Games, Inc. All Rights Reserved.
+
+#pragma once
+
+#include "Containers/ContainersFwd.h"
+#include "Containers/UnrealString.h"
+#include "NNETypes.h"
+
+namespace UE::NNERuntimeRDG::Private
+{
+	/** Concrete tensor with data accessible by graph scheduling */
+	class FTensor
+	{
+	protected:
+		FString				Name;
+		ENNETensorDataType	DataType;
+		NNE::FTensorShape	Shape;
+		TArray<uint8>		PreparedData;
+		uint64				DataSize = 0;
+		uint32				Volume = 0;
+
+		FTensor() = default;
+
+	public:
+		const FString& GetName() const
+		{
+			return Name;
+		}
+
+		ENNETensorDataType GetDataType() const
+		{
+			return DataType;
+		}
+
+		uint32 GetElementByteSize() const
+		{
+			return NNE::GetTensorDataTypeSizeInBytes(DataType);
+		}
+
+		const NNE::FTensorShape& GetShape() const
+		{
+			return Shape;
+		}
+
+		template <typename T> TConstArrayView<T> GetPreparedData() const
+		{
+			const T* DataPtr = reinterpret_cast<const T*>(PreparedData.GetData());
+			const int32 ElemSize = sizeof(T);
+
+			check(PreparedData.Num() % ElemSize == 0);
+			return MakeArrayView(DataPtr, PreparedData.Num() / ElemSize);
+		}
+
+		void SetShape(const NNE::FTensorShape& InShape)
+		{
+			checkf(!HasPreparedData(), TEXT("Shape cannot be changed once data as been set."));
+			check(InShape.Volume() <= TNumericLimits<uint32>::Max());
+			Shape = InShape;
+			Volume = InShape.Volume();
+			DataSize = (uint64)NNE::GetTensorDataTypeSizeInBytes(DataType) * Volume;
+		}
+
+		template <typename T> void SetPreparedData(TConstArrayView<T> Data)
+		{
+			const uint8* DataPtr = reinterpret_cast<const uint8*>(Data.GetData());
+			TConstArrayView<uint8> DataAsByte = MakeArrayView(DataPtr, Data.Num() * sizeof(T));
+			
+			checkf(DataAsByte.Num() == DataSize, TEXT("Incorrect data size, it should match tensor shape and data type."));
+			PreparedData.Reset();
+			PreparedData.Append(DataAsByte);
+		}
+
+		bool IsEmpty() const
+		{
+			bool IsEmpty = GetDataType() == ENNETensorDataType::None;
+			ensureMsgf(!IsEmpty || (Shape.Rank() == 1 && Shape.GetData()[0] == 0), TEXT("Empty tensor should have a shape of [0]."));
+			return IsEmpty;
+		}
+
+		bool HasPreparedData() const
+		{
+			return !PreparedData.IsEmpty();
+		}
+
+		bool IsConstant() const
+		{
+			return Volume == 0 || HasPreparedData();
+		}
+
+		uint32 GetVolume() const
+		{
+			return Volume;
+		}
+
+		uint64 GetDataSize() const
+		{
+			return DataSize;
+		}
+
+		static FTensor Make(const FString& Name, const NNE::FTensorShape& Shape, ENNETensorDataType DataType)
+		{
+			FTensor Tensor;
+			Tensor.Name = Name;
+			Tensor.DataType = DataType;
+			Tensor.SetShape(Shape);
+			return Tensor;
+		}
+
+		static FTensor Make(const NNE::FTensorDesc& TensorDesc, const NNE::FTensorShape& Shape)
+		{
+			check(Shape.IsCompatibleWith(TensorDesc.GetShape()));
+			return Make(TensorDesc.GetName(), Shape, TensorDesc.GetDataType());
+		}
+
+		static FTensor MakeFromSymbolicDesc(const NNE::FTensorDesc& TensorDesc)
+		{
+			return Make(TensorDesc.GetName(), NNE::FTensorShape::MakeFromSymbolic(TensorDesc.GetShape()), TensorDesc.GetDataType());
+		}
+	};
+
+	using FTensorRef = FTensor*;
+
+} // namespace UE::NNERuntimeRDG::Private
\ No newline at end of file
diff --git a/Engine/Plugins/Experimental/NNERuntimeRDG/Source/NNERuntimeRDG/Private/Tests/NNERuntimeRDGUnitTestHelper.cpp b/Engine/Plugins/Experimental/NNERuntimeRDG/Source/NNERuntimeRDG/Private/Tests/NNERuntimeRDGUnitTestHelper.cpp
--- a/Engine/Plugins/Experimental/NNERuntimeRDG/Source/NNERuntimeRDG/Private/Tests/NNERuntimeRDGUnitTestHelper.cpp
+++ b/Engine/Plugins/Experimental/NNERuntimeRDG/Source/NNERuntimeRDG/Private/Tests/NNERuntimeRDGUnitTestHelper.cpp
@@ -2,8 +2,8 @@
 
 #include "NNERuntimeRDGUnitTestHelper.h"
 #include "NNETypes.h"
-#include "NNERuntimeRDGHelperElementWiseBinary.h"
-#include "NNERuntimeRDGHelperElementWiseUnary.h"
+#include "Helper/NNERuntimeRDGHelperElementWiseBinary.h"
+#include "Helper/NNERuntimeRDGHelperElementWiseUnary.h"
 
 #if WITH_DEV_AUTOMATION_TESTS
 #include "Misc/AutomationTest.h"
@@ -12,8 +12,6 @@ namespace UE::NNERuntimeRDG::Private::Test
 {
 
 using namespace NNE;
-using namespace NNE::Internal;
-using namespace UE::NNERuntimeRDG::Internal;
 
 bool FUnitTestBase::RunTest(const FString& Parameters) { return false; }
 
diff --git a/Engine/Plugins/Experimental/NNERuntimeRDG/Source/NNERuntimeRDG/Private/Tests/NNERuntimeRDGUnitTestHelper.h b/Engine/Plugins/Experimental/NNERuntimeRDG/Source/NNERuntimeRDG/Private/Tests/NNERuntimeRDGUnitTestHelper.h
--- a/Engine/Plugins/Experimental/NNERuntimeRDG/Source/NNERuntimeRDG/Private/Tests/NNERuntimeRDGUnitTestHelper.h
+++ b/Engine/Plugins/Experimental/NNERuntimeRDG/Source/NNERuntimeRDG/Private/Tests/NNERuntimeRDGUnitTestHelper.h
@@ -3,7 +3,7 @@
 #pragma once
 
 #include "CoreMinimal.h"
-#include "NNETensor.h"
+#include "NNERuntimeRDGTensor.h"
 #include "NNEHlslShadersOperator.h"
 
 #if WITH_DEV_AUTOMATION_TESTS
@@ -32,11 +32,11 @@ namespace UE::NNERuntimeRDG::Private::Test
 		virtual bool RunTest(const FString& Parameter) override;
 
 		//Shape inference unit test helpers
-		static NNE::Internal::FTensor MakeTensor(FString Name, TConstArrayView<uint32> Shape, ENNETensorDataType DataType = ENNETensorDataType::Float);
-		static NNE::Internal::FTensor MakeConstTensor(FString Name, TConstArrayView<uint32> Shape, TConstArrayView<float> Data);
-		static NNE::Internal::FTensor MakeConstTensorFloat16(FString Name, TConstArrayView<uint32> Shape, TConstArrayView<FFloat16> Data);
-		static NNE::Internal::FTensor MakeConstTensorInt32(FString Name, TConstArrayView<uint32> Shape, TConstArrayView<int32> Data);
-		static NNE::Internal::FTensor MakeConstTensorInt64(FString Name, TConstArrayView<uint32> Shape, TConstArrayView<int64> Data);
+		static FTensor MakeTensor(FString Name, TConstArrayView<uint32> Shape, ENNETensorDataType DataType = ENNETensorDataType::Float);
+		static FTensor MakeConstTensor(FString Name, TConstArrayView<uint32> Shape, TConstArrayView<float> Data);
+		static FTensor MakeConstTensorFloat16(FString Name, TConstArrayView<uint32> Shape, TConstArrayView<FFloat16> Data);
+		static FTensor MakeConstTensorInt32(FString Name, TConstArrayView<uint32> Shape, TConstArrayView<int32> Data);
+		static FTensor MakeConstTensorInt64(FString Name, TConstArrayView<uint32> Shape, TConstArrayView<int64> Data);
 		bool TestUnaryOutputIsOnlyComputedWhenItShould(NNEHlslShaders::Internal::EElementWiseUnaryOperatorType OpType);
 		bool TestBinaryOutputIsOnlyComputedWhenItShould(NNEHlslShaders::Internal::EElementWiseBinaryOperatorType OpType);
 	};
diff --git a/Engine/Plugins/Experimental/NNERuntimeRDG/Source/NNERuntimeRDG/Private/Tests/NNERuntimeRDGUnitTestOpBinary.cpp b/Engine/Plugins/Experimental/NNERuntimeRDG/Source/NNERuntimeRDG/Private/Tests/NNERuntimeRDGUnitTestOpBinary.cpp
--- a/Engine/Plugins/Experimental/NNERuntimeRDG/Source/NNERuntimeRDG/Private/Tests/NNERuntimeRDGUnitTestOpBinary.cpp
+++ b/Engine/Plugins/Experimental/NNERuntimeRDG/Source/NNERuntimeRDG/Private/Tests/NNERuntimeRDGUnitTestOpBinary.cpp
@@ -1,10 +1,10 @@
 // Copyright Epic Games, Inc. All Rights Reserved.
 
 #include "CoreMinimal.h"
-#include "NNETensor.h"
-#include "NNETypes.h"
-#include "NNERuntimeRDGHelperElementWiseBinary.h"
+#include "Helper/NNERuntimeRDGHelperElementWiseBinary.h"
+#include "NNERuntimeRDGTensor.h"
 #include "NNERuntimeRDGUnitTestHelper.h"
+#include "NNETypes.h"
 
 
 #if WITH_DEV_AUTOMATION_TESTS
@@ -17,8 +17,7 @@ namespace UE::NNERuntimeRDG::Private::Test::BinaryOp
 #if WITH_DEV_AUTOMATION_TESTS
 	
 	using namespace NNE;
-	using namespace NNE::Internal;
-	using namespace UE::NNERuntimeRDG::Internal;
+
 	using EElementWiseBinaryOperatorType = NNEHlslShaders::Internal::EElementWiseBinaryOperatorType;
 
 	IMPLEMENT_NNE_SHAPEINFERENCEHELPER_UNIT_AUTOMATION_TEST(FElementWiseBinaryCPUHelperAdd, "System.Engine.MachineLearning.NNE.RDG.UnitTest.BinaryHelper.Add");
diff --git a/Engine/Plugins/Experimental/NNERuntimeRDG/Source/NNERuntimeRDG/Private/Tests/NNERuntimeRDGUnitTestOpCast.cpp b/Engine/Plugins/Experimental/NNERuntimeRDG/Source/NNERuntimeRDG/Private/Tests/NNERuntimeRDGUnitTestOpCast.cpp
--- a/Engine/Plugins/Experimental/NNERuntimeRDG/Source/NNERuntimeRDG/Private/Tests/NNERuntimeRDGUnitTestOpCast.cpp
+++ b/Engine/Plugins/Experimental/NNERuntimeRDG/Source/NNERuntimeRDG/Private/Tests/NNERuntimeRDGUnitTestOpCast.cpp
@@ -1,10 +1,10 @@
 // Copyright Epic Games, Inc. All Rights Reserved.
 
 #include "CoreMinimal.h"
-#include "NNETensor.h"
-#include "NNETypes.h"
-#include "NNERuntimeRDGHelperCast.h"
+#include "Helper/NNERuntimeRDGHelperCast.h"
+#include "NNERuntimeRDGTensor.h"
 #include "NNERuntimeRDGUnitTestHelper.h"
+#include "NNETypes.h"
 
 #if WITH_DEV_AUTOMATION_TESTS
 #include "Misc/AutomationTest.h"
@@ -15,9 +15,6 @@ namespace UE::NNERuntimeRDG::Private::Test::CastOp
 
 #if WITH_DEV_AUTOMATION_TESTS
 
-	using namespace NNE::Internal;
-	using namespace UE::NNERuntimeRDG::Internal;
-
 	IMPLEMENT_NNE_SHAPEINFERENCEHELPER_UNIT_AUTOMATION_TEST(FCastCPUHelperConstOuput, "System.Engine.MachineLearning.NNE.RDG.UnitTest.CastHelper.ConstOutput");
 	bool FCastCPUHelperConstOuput::RunTest(const FString& Parameter)
 	{
diff --git a/Engine/Plugins/Experimental/NNERuntimeRDG/Source/NNERuntimeRDG/Private/Tests/NNERuntimeRDGUnitTestOpConcat.cpp b/Engine/Plugins/Experimental/NNERuntimeRDG/Source/NNERuntimeRDG/Private/Tests/NNERuntimeRDGUnitTestOpConcat.cpp
--- a/Engine/Plugins/Experimental/NNERuntimeRDG/Source/NNERuntimeRDG/Private/Tests/NNERuntimeRDGUnitTestOpConcat.cpp
+++ b/Engine/Plugins/Experimental/NNERuntimeRDG/Source/NNERuntimeRDG/Private/Tests/NNERuntimeRDGUnitTestOpConcat.cpp
@@ -1,10 +1,10 @@
 // Copyright Epic Games, Inc. All Rights Reserved.
 
 #include "CoreMinimal.h"
-#include "NNETensor.h"
-#include "NNETypes.h"
-#include "NNERuntimeRDGHelperConcat.h"
+#include "Helper/NNERuntimeRDGHelperConcat.h"
+#include "NNERuntimeRDGTensor.h"
 #include "NNERuntimeRDGUnitTestHelper.h"
+#include "NNETypes.h"
 
 #if WITH_DEV_AUTOMATION_TESTS
 #include "Misc/AutomationTest.h"
@@ -15,9 +15,6 @@ namespace UE::NNERuntimeRDG::Private::Test::ConcatOp
 
 #if WITH_DEV_AUTOMATION_TESTS
 
-	using namespace NNE::Internal;
-	using namespace UE::NNERuntimeRDG::Internal;
-
 	IMPLEMENT_NNE_SHAPEINFERENCEHELPER_UNIT_AUTOMATION_TEST(FConcatCPUHelperConstOuput, "System.Engine.MachineLearning.NNE.RDG.UnitTest.ConcatHelper.ConstOutput");
 	bool FConcatCPUHelperConstOuput::RunTest(const FString& Parameter)
 	{
diff --git a/Engine/Plugins/Experimental/NNERuntimeRDG/Source/NNERuntimeRDG/Private/Tests/NNERuntimeRDGUnitTestOpGather.cpp b/Engine/Plugins/Experimental/NNERuntimeRDG/Source/NNERuntimeRDG/Private/Tests/NNERuntimeRDGUnitTestOpGather.cpp
--- a/Engine/Plugins/Experimental/NNERuntimeRDG/Source/NNERuntimeRDG/Private/Tests/NNERuntimeRDGUnitTestOpGather.cpp
+++ b/Engine/Plugins/Experimental/NNERuntimeRDG/Source/NNERuntimeRDG/Private/Tests/NNERuntimeRDGUnitTestOpGather.cpp
@@ -1,10 +1,11 @@
 // Copyright Epic Games, Inc. All Rights Reserved.
 
 #include "CoreMinimal.h"
-#include "NNETensor.h"
-#include "NNETypes.h"
-#include "NNERuntimeRDGHelperGather.h"
+
+#include "Helper/NNERuntimeRDGHelperGather.h"
 #include "NNERuntimeRDGUnitTestHelper.h"
+#include "NNERuntimeRDGTensor.h"
+#include "NNETypes.h"
 
 #if WITH_DEV_AUTOMATION_TESTS
 #include "Misc/AutomationTest.h"
@@ -15,9 +16,6 @@ namespace UE::NNERuntimeRDG::Private::Test::GatherOp
 
 #if WITH_DEV_AUTOMATION_TESTS
 
-	using namespace NNE::Internal;
-	using namespace UE::NNERuntimeRDG::Internal;
-
 	IMPLEMENT_NNE_SHAPEINFERENCEHELPER_UNIT_AUTOMATION_TEST(FGatherCPUHelperConstOuput, "System.Engine.MachineLearning.NNE.RDG.UnitTest.GatherHelper.ConstOutput");
 	bool FGatherCPUHelperConstOuput::RunTest(const FString& Parameter)
 	{
diff --git a/Engine/Plugins/Experimental/NNERuntimeRDG/Source/NNERuntimeRDG/Private/Tests/NNERuntimeRDGUnitTestOpSlice.cpp b/Engine/Plugins/Experimental/NNERuntimeRDG/Source/NNERuntimeRDG/Private/Tests/NNERuntimeRDGUnitTestOpSlice.cpp
--- a/Engine/Plugins/Experimental/NNERuntimeRDG/Source/NNERuntimeRDG/Private/Tests/NNERuntimeRDGUnitTestOpSlice.cpp
+++ b/Engine/Plugins/Experimental/NNERuntimeRDG/Source/NNERuntimeRDG/Private/Tests/NNERuntimeRDGUnitTestOpSlice.cpp
@@ -1,10 +1,10 @@
 // Copyright Epic Games, Inc. All Rights Reserved.
 
 #include "CoreMinimal.h"
-#include "NNETensor.h"
-#include "NNETypes.h"
-#include "NNERuntimeRDGHelperSlice.h"
+#include "Helper/NNERuntimeRDGHelperSlice.h"
 #include "NNERuntimeRDGUnitTestHelper.h"
+#include "NNERuntimeRDGTensor.h"
+#include "NNETypes.h"
 
 #if WITH_DEV_AUTOMATION_TESTS
 #include "Misc/AutomationTest.h"
@@ -15,14 +15,10 @@ namespace UE::NNERuntimeRDG::Private::Test::SliceOp
 
 #if WITH_DEV_AUTOMATION_TESTS
 
-	using namespace NNE::Internal;
-	
-
 	namespace
 	{
 		void CallSliceApply(const FTensor& InputTensor, FTensor& OutputTensor, TConstArrayView<int32> Starts, TConstArrayView<int32> Steps)
 		{
-			using namespace UE::NNERuntimeRDG::Internal;
 			CPUHelper::Slice::Apply(InputTensor, OutputTensor, Starts, Steps);
 		}
 
diff --git a/Engine/Plugins/Experimental/NNERuntimeRDG/Source/NNERuntimeRDG/Private/Tests/NNERuntimeRDGUnitTestOpTranspose.cpp b/Engine/Plugins/Experimental/NNERuntimeRDG/Source/NNERuntimeRDG/Private/Tests/NNERuntimeRDGUnitTestOpTranspose.cpp
--- a/Engine/Plugins/Experimental/NNERuntimeRDG/Source/NNERuntimeRDG/Private/Tests/NNERuntimeRDGUnitTestOpTranspose.cpp
+++ b/Engine/Plugins/Experimental/NNERuntimeRDG/Source/NNERuntimeRDG/Private/Tests/NNERuntimeRDGUnitTestOpTranspose.cpp
@@ -1,10 +1,10 @@
 // Copyright Epic Games, Inc. All Rights Reserved.
 
 #include "CoreMinimal.h"
-#include "NNETensor.h"
-#include "NNETypes.h"
-#include "NNERuntimeRDGHelperTranspose.h"
+#include "Helper/NNERuntimeRDGHelperTranspose.h"
 #include "NNERuntimeRDGUnitTestHelper.h"
+#include "NNERuntimeRDGTensor.h"
+#include "NNETypes.h"
 
 #if WITH_DEV_AUTOMATION_TESTS
 #include "Misc/AutomationTest.h"
@@ -15,9 +15,6 @@ namespace UE::NNERuntimeRDG::Private::Test::TransposeOp
 
 #if WITH_DEV_AUTOMATION_TESTS
 
-	using namespace NNE::Internal;
-	using namespace UE::NNERuntimeRDG::Internal;
-
 	IMPLEMENT_NNE_SHAPEINFERENCEHELPER_UNIT_AUTOMATION_TEST(FTransposeCPUHelperIdentity, "System.Engine.MachineLearning.NNE.RDG.UnitTest.TransposeHelper.Identity");
 	bool FTransposeCPUHelperIdentity::RunTest(const FString& Parameter)
 	{
diff --git a/Engine/Plugins/Experimental/NNERuntimeRDG/Source/NNERuntimeRDG/Private/Tests/NNERuntimeRDGUnitTestOpUnary.cpp b/Engine/Plugins/Experimental/NNERuntimeRDG/Source/NNERuntimeRDG/Private/Tests/NNERuntimeRDGUnitTestOpUnary.cpp
--- a/Engine/Plugins/Experimental/NNERuntimeRDG/Source/NNERuntimeRDG/Private/Tests/NNERuntimeRDGUnitTestOpUnary.cpp
+++ b/Engine/Plugins/Experimental/NNERuntimeRDG/Source/NNERuntimeRDG/Private/Tests/NNERuntimeRDGUnitTestOpUnary.cpp
@@ -2,10 +2,11 @@
 
 #include "CoreMinimal.h"
 #include "Misc/EnumRange.h"
-#include "NNETensor.h"
-#include "NNETypes.h"
-#include "NNERuntimeRDGHelperElementWiseUnary.h"
+
+#include "Helper/NNERuntimeRDGHelperElementWiseUnary.h"
 #include "NNERuntimeRDGUnitTestHelper.h"
+#include "NNERuntimeRDGTensor.h"
+#include "NNETypes.h"
 
 #if WITH_DEV_AUTOMATION_TESTS
 #include "Misc/AutomationTest.h"
@@ -19,8 +20,6 @@ namespace UE::NNERuntimeRDG::Private::Test::UnaryOp
 #if WITH_DEV_AUTOMATION_TESTS
 
 	using namespace NNE;
-	using namespace NNE::Internal;
-	using namespace UE::NNERuntimeRDG::Internal;
 	using EElementWiseUnaryOperatorType = NNEHlslShaders::Internal::EElementWiseUnaryOperatorType;
 
 	IMPLEMENT_NNE_SHAPEINFERENCEHELPER_UNIT_AUTOMATION_TEST(FElementWiseUnaryCPUHelperConstOuput, "System.Engine.MachineLearning.NNE.RDG.UnitTest.UnaryHelper.ConstOutput");
diff --git a/Engine/Plugins/Experimental/NNERuntimeRDG/Source/NNERuntimeRDGData/Internal/NNERuntimeRDGDataAttributeTensor.h b/Engine/Plugins/Experimental/NNERuntimeRDG/Source/NNERuntimeRDGData/Internal/NNERuntimeRDGDataAttributeTensor.h
--- a/Engine/Plugins/Experimental/NNERuntimeRDG/Source/NNERuntimeRDGData/Internal/NNERuntimeRDGDataAttributeTensor.h
+++ b/Engine/Plugins/Experimental/NNERuntimeRDG/Source/NNERuntimeRDGData/Internal/NNERuntimeRDGDataAttributeTensor.h
@@ -3,7 +3,6 @@
 #pragma once
 
 #include "Containers/ArrayView.h"
-#include "NNETensor.h"
 #include "NNETypes.h"
 #include "Serialization/Archive.h"
 
@@ -23,10 +22,14 @@ namespace UE::NNERuntimeRDGData::Internal
 			return DataType;
 		}
 
-		void FillFTensorWithShapeAndData(NNE::Internal::FTensor& Tensor) const
+		TConstArrayView<uint8> GetData() const
 		{
-			Tensor.SetShape(NNE::FTensorShape::Make(Shape));
-			Tensor.SetPreparedData(MakeConstArrayView(Data));
+			return Data;
+		}
+
+		NNE::FTensorShape GetShape() const
+		{
+			return NNE::FTensorShape::Make(Shape);
 		}
 
 		static FAttributeTensor Make(const NNE::FTensorShape& Shape,
diff --git a/Engine/Plugins/Experimental/NNERuntimeRDG/Source/NNERuntimeRDGUtils/Private/NNERuntimeRDGUtilsModelBuilder.h b/Engine/Plugins/Experimental/NNERuntimeRDG/Source/NNERuntimeRDGUtils/Private/NNERuntimeRDGUtilsModelBuilder.h
--- a/Engine/Plugins/Experimental/NNERuntimeRDG/Source/NNERuntimeRDGUtils/Private/NNERuntimeRDGUtilsModelBuilder.h
+++ b/Engine/Plugins/Experimental/NNERuntimeRDG/Source/NNERuntimeRDGUtils/Private/NNERuntimeRDGUtilsModelBuilder.h
@@ -3,7 +3,7 @@
 #pragma once
 
 #include "CoreMinimal.h"
-#include "NNETensor.h"
+#include "Misc/OptionalFwd.h"
 #include "NNETypes.h"
 #include "Templates/UniquePtr.h"
 
diff --git a/Engine/Plugins/Experimental/NNERuntimeRDG/Source/NNERuntimeRDGUtils/Private/NNERuntimeRDGUtilsModelOptimizerNNE.cpp b/Engine/Plugins/Experimental/NNERuntimeRDG/Source/NNERuntimeRDGUtils/Private/NNERuntimeRDGUtilsModelOptimizerNNE.cpp
--- a/Engine/Plugins/Experimental/NNERuntimeRDG/Source/NNERuntimeRDGUtils/Private/NNERuntimeRDGUtilsModelOptimizerNNE.cpp
+++ b/Engine/Plugins/Experimental/NNERuntimeRDG/Source/NNERuntimeRDGUtils/Private/NNERuntimeRDGUtilsModelOptimizerNNE.cpp
@@ -9,7 +9,6 @@
 #include "NNERuntimeRDGUtilsModelBuilderNNE.h"
 #include "NNERuntimeRDGUtilsModelOptimizerONNX.h"
 #include "NNERuntimeRDGUtilsHelpers.h"
-#include "NNETensor.h"
 #include "NNETypes.h"
 
 THIRD_PARTY_INCLUDES_START
