[NNE] HLSL runtime only register base runtime interface when inference is not supported.

Problems:
On unsupported platform the RDG interface is always registered for HLSL, thus allowing one to call inference and get errors from it.
This does not follow other plugin NNE design and is unclear in behavior.

Solutions:
Only register the NNE RDG interface when inference is supported.

#tests all nneqa tests on windows dev editor + nneqa smoketest on console debug and mac editor debug + verify inference interface is not registered on SM5 windows dev editor + style transfer hlsl win & console
#rnx

diff --git a/Engine/Plugins/Experimental/NNERuntimeRDG/Source/NNERuntimeRDG/Private/NNERuntimeRDGHlsl.cpp b/Engine/Plugins/Experimental/NNERuntimeRDG/Source/NNERuntimeRDG/Private/NNERuntimeRDGHlsl.cpp
--- a/Engine/Plugins/Experimental/NNERuntimeRDG/Source/NNERuntimeRDG/Private/NNERuntimeRDGHlsl.cpp
+++ b/Engine/Plugins/Experimental/NNERuntimeRDG/Source/NNERuntimeRDG/Private/NNERuntimeRDGHlsl.cpp
@@ -165,39 +165,12 @@ namespace UE::NNERuntimeRDG::Private::Hlsl
 } // namespace UE::NNERuntimeRDG::Private::Hlsl
 
 
-bool UNNERuntimeRDGHlslImpl::IsCurrentPlatformSupported()
-{
-	bool bResult = true;
-
-#ifndef NNE_FORCE_HARDWARE_SUPPORTS_HLSL
-	if(GMaxRHIFeatureLevel < ERHIFeatureLevel::SM5)
-	{
-		UE_LOG(LogNNERuntimeRDGHlsl, Display, TEXT("Minimum feature level required is SM5 for current RHI platform."));
-		bResult = false;
-	}
-
-	if(!GRHISupportsWaveOperations)
-	{
-		UE_LOG(LogNNERuntimeRDGHlsl, Display, TEXT("Current RHI platform doesn't support wave operations."));
-		bResult = false;
-	}
-
-	if(!GRHIGlobals.SupportsNative16BitOps)
-	{
-		UE_LOG(LogNNERuntimeRDGHlsl, Display, TEXT("Current RHI platform doesn't support native 16-bit operations."));
-		bResult = false;
-	}
-#endif
-
-	return bResult;
-}
-
 UNNERuntimeRDGHlslImpl::ECanCreateModelDataStatus UNNERuntimeRDGHlslImpl::CanCreateModelData(const FString& FileType, TConstArrayView64<uint8> FileData, const TMap<FString, TConstArrayView64<uint8>>& AdditionalFileData, const FGuid& FileId, const ITargetPlatform* TargetPlatform) const
 {
 	return Details::CheckCanCreateModelData(/*bShouldLog*/ false , FileType, FileData, AdditionalFileData, FileId, TargetPlatform);
 }
 
-UNNERuntimeRDGHlslImpl::ECanCreateModelRDGStatus UNNERuntimeRDGHlslImpl::CanCreateModelRDG(const TObjectPtr<UNNEModelData> ModelData) const
+UNNERuntimeRDGHlslImplRDG::ECanCreateModelRDGStatus UNNERuntimeRDGHlslImplRDG::CanCreateModelRDG(const TObjectPtr<UNNEModelData> ModelData) const
 {
 	int32 GuidSize = sizeof(GUID);
 	int32 VersionSize = sizeof(Version);
@@ -217,8 +190,6 @@ UNNERuntimeRDGHlslImpl::ECanCreateModelRDGStatus UNNERuntimeRDGHlslImpl::CanCrea
 	bool bResult = FGenericPlatformMemory::Memcmp(&(Data[0]), &(GUID), GuidSize) == 0;
 	bResult &= FGenericPlatformMemory::Memcmp(&(Data[GuidSize]), &(Version), VersionSize) == 0;
 
-	bResult &= IsCurrentPlatformSupported();
-
 	return bResult ? ECanCreateModelRDGStatus::Ok : ECanCreateModelRDGStatus::Fail;
 };
 
@@ -257,7 +228,7 @@ FString UNNERuntimeRDGHlslImpl::GetModelDataIdentifier(const FString& FileType,
 	return FileId.ToString(EGuidFormats::Digits) + "-" + UNNERuntimeRDGHlslImpl::GUID.ToString(EGuidFormats::Digits) + "-" + FString::FromInt(UNNERuntimeRDGHlslImpl::Version);
 }
 
-TSharedPtr<UE::NNE::IModelRDG> UNNERuntimeRDGHlslImpl::CreateModelRDG(const TObjectPtr<UNNEModelData> ModelData)
+TSharedPtr<UE::NNE::IModelRDG> UNNERuntimeRDGHlslImplRDG::CreateModelRDG(const TObjectPtr<UNNEModelData> ModelData)
 {
 	if (CanCreateModelRDG(ModelData) != ECanCreateModelRDGStatus::Ok)
 	{
diff --git a/Engine/Plugins/Experimental/NNERuntimeRDG/Source/NNERuntimeRDG/Private/NNERuntimeRDGHlsl.h b/Engine/Plugins/Experimental/NNERuntimeRDG/Source/NNERuntimeRDG/Private/NNERuntimeRDGHlsl.h
--- a/Engine/Plugins/Experimental/NNERuntimeRDG/Source/NNERuntimeRDG/Private/NNERuntimeRDGHlsl.h
+++ b/Engine/Plugins/Experimental/NNERuntimeRDG/Source/NNERuntimeRDG/Private/NNERuntimeRDGHlsl.h
@@ -13,7 +13,7 @@
 #include "NNERuntimeRDGHlsl.generated.h"
 
 UCLASS()
-class UNNERuntimeRDGHlslImpl : public UObject, public INNERuntime, public INNERuntimeRDG
+class UNNERuntimeRDGHlslImpl : public UObject, public INNERuntime
 {
 	GENERATED_BODY()
 
@@ -28,11 +28,18 @@ public:
 
 	virtual FString GetRuntimeName() const override { return TEXT("NNERuntimeRDGHlsl"); };
 
-	static bool IsCurrentPlatformSupported();
-
 	virtual ECanCreateModelDataStatus CanCreateModelData(const FString& FileType, TConstArrayView64<uint8> FileData, const TMap<FString, TConstArrayView64<uint8>>& AdditionalFileData, const FGuid& FileId, const ITargetPlatform* TargetPlatform) const override;
 	virtual TSharedPtr<UE::NNE::FSharedModelData> CreateModelData(const FString& FileType, TConstArrayView64<uint8> FileData, const TMap<FString, TConstArrayView64<uint8>>& AdditionalFileData, const FGuid& FileId, const ITargetPlatform* TargetPlatform) override;
 	virtual FString GetModelDataIdentifier(const FString& FileType, TConstArrayView64<uint8> FileData, const TMap<FString, TConstArrayView64<uint8>>& AdditionalFileData, const FGuid& FileId, const ITargetPlatform* TargetPlatform) const override;
+};
+
+UCLASS()
+class UNNERuntimeRDGHlslImplRDG : public UNNERuntimeRDGHlslImpl, public INNERuntimeRDG
+{
+	GENERATED_BODY()
+
+public:
+	virtual ~UNNERuntimeRDGHlslImplRDG() {}
 
 	virtual ECanCreateModelRDGStatus CanCreateModelRDG(const TObjectPtr<UNNEModelData> ModelData) const override;
 	virtual TSharedPtr<UE::NNE::IModelRDG> CreateModelRDG(const TObjectPtr<UNNEModelData> ModelData) override;
diff --git a/Engine/Plugins/Experimental/NNERuntimeRDG/Source/NNERuntimeRDG/Private/NNERuntimeRDGModule.cpp b/Engine/Plugins/Experimental/NNERuntimeRDG/Source/NNERuntimeRDG/Private/NNERuntimeRDGModule.cpp
--- a/Engine/Plugins/Experimental/NNERuntimeRDG/Source/NNERuntimeRDG/Private/NNERuntimeRDGModule.cpp
+++ b/Engine/Plugins/Experimental/NNERuntimeRDG/Source/NNERuntimeRDG/Private/NNERuntimeRDGModule.cpp
@@ -2,12 +2,13 @@
 
 #include "NNERuntimeRDGModule.h"
 
+#include "DataDrivenShaderPlatformInfo.h"
 #include "HAL/IConsoleManager.h"
 #include "NNE.h"
 #include "NNEHlslShadersLog.h"
 #include "NNERuntimeRDGHlsl.h"
+#include "RHI.h"
 #include "UObject/WeakInterfacePtr.h"
-#include "DataDrivenShaderPlatformInfo.h"
 
 static TAutoConsoleVariable<int32> CVarHlslModelOptimization(
 	TEXT("nne.hlsl.ModelOptimization"),
@@ -16,9 +17,39 @@ static TAutoConsoleVariable<int32> CVarHlslModelOptimization(
 	TEXT(" 0: disabled\n")
 	TEXT(" 1: enabled (default)")
 #if !WITH_EDITOR
-	,ECVF_ReadOnly
+	, ECVF_ReadOnly
+#endif
+);
+
+namespace UE::NNERuntimeRDG::Private::Details
+{
+	bool IsInferenceSupported()
+	{
+		bool bResult = true;
+
+#ifndef NNE_FORCE_HARDWARE_SUPPORTS_HLSL
+		if (GMaxRHIFeatureLevel < ERHIFeatureLevel::SM5)
+		{
+			UE_LOG(LogNNERuntimeRDGHlsl, Display, TEXT("Minimum feature level required is SM5 for current RHI platform."));
+			bResult = false;
+		}
+
+		if (!GRHISupportsWaveOperations)
+		{
+			UE_LOG(LogNNERuntimeRDGHlsl, Display, TEXT("Current RHI platform doesn't support wave operations."));
+			bResult = false;
+		}
+
+		if (!GRHIGlobals.SupportsNative16BitOps)
+		{
+			UE_LOG(LogNNERuntimeRDGHlsl, Display, TEXT("Current RHI platform doesn't support native 16-bit operations."));
+			bResult = false;
+		}
 #endif
-	);
+
+		return bResult;
+	}
+}
 
 void FNNERuntimeRDGModule::RegisterRuntime()
 {
@@ -35,31 +66,40 @@ void FNNERuntimeRDGModule::RegisterRuntime()
 
 void FNNERuntimeRDGModule::StartupModule()
 {
-#if WITH_EDITOR
-	#ifdef WITH_NNE_RUNTIME_HLSL 
-		RegisterRuntime();
-	#endif
-#else
-	#ifdef WITH_NNE_RUNTIME_HLSL 
-		if(FDataDrivenShaderPlatformInfo::GetSupportsNNEShaders(GMaxRHIShaderPlatform))
+#ifdef WITH_NNE_RUNTIME_HLSL 
+	if (FDataDrivenShaderPlatformInfo::GetSupportsNNEShaders(GMaxRHIShaderPlatform))
+	{
+		if (UE::NNERuntimeRDG::Private::Details::IsInferenceSupported())
 		{
-			if(UNNERuntimeRDGHlslImpl::IsCurrentPlatformSupported())
-			{
-				RegisterRuntime();
-			}
-			else
-			{
-				UE_LOG(LogNNERuntimeRDGHlsl, Display, TEXT("Not registering runtime because current hardware is incompatible, consider bypassing by setting the define NNE_FORCE_HARDWARE_SUPPORTS_HLSL."));
-			}
+			NNERuntimeRDGHlsl = NewObject<UNNERuntimeRDGHlslImplRDG>();
 		}
 		else
 		{
-			UE_LOG(LogNNERuntimeRDGHlsl, Display, TEXT("Not registering runtime because current RHI shader platform is not enabled, consider setting the flag bSupportsNNEShaders in DataDrivenPlatformInfo."));
+			UE_LOG(LogNNERuntimeRDGHlsl, Display, TEXT("Not registering inference for runtime because current hardware is incompatible, consider bypassing by setting the define NNE_FORCE_HARDWARE_SUPPORTS_HLSL."));
 		}
-	#else
-		UE_LOG(LogNNERuntimeRDGHlsl, Display, TEXT("Not registering runtime as platform is not enabled, if needed set define WITH_NNE_RUNTIME_HLSL."));
-	#endif
-#endif
+	}
+	else
+	{
+		UE_LOG(LogNNERuntimeRDGHlsl, Display, TEXT("Not registering inference for runtime because current RHI shader platform is not enabled, consider setting the flag bSupportsNNEShaders in DataDrivenPlatformInfo."));
+	}
+
+	#if WITH_EDITOR
+	if (!NNERuntimeRDGHlsl.IsValid())
+	{
+		// We can always cook in editor
+		NNERuntimeRDGHlsl = NewObject<UNNERuntimeRDGHlslImpl>();
+	}
+	#endif // WITH_EDITOR
+
+	if (NNERuntimeRDGHlsl.IsValid())
+	{
+		TWeakInterfacePtr<INNERuntime> RuntimeInterface(NNERuntimeRDGHlsl.Get());
+
+		NNERuntimeRDGHlsl->Init();
+		NNERuntimeRDGHlsl->AddToRoot();
+		UE::NNE::RegisterRuntime(RuntimeInterface);
+	}
+#endif // WITH_NNE_RUNTIME_HLSL
 }
 
 void FNNERuntimeRDGModule::ShutdownModule()
