Update numeric limits to use std lib instead of macro because it fails to compile on newer Windows 11 SDKs

diff --git a/Engine/Plugins/Experimental/NNERuntimeRDG/Source/NNERuntimeRDG/Private/Tests/NNERuntimeRDGUnitTestOpUnary.cpp b/Engine/Plugins/Experimental/NNERuntimeRDG/Source/NNERuntimeRDG/Private/Tests/NNERuntimeRDGUnitTestOpUnary.cpp
--- a/Engine/Plugins/Experimental/NNERuntimeRDG/Source/NNERuntimeRDG/Private/Tests/NNERuntimeRDGUnitTestOpUnary.cpp
+++ b/Engine/Plugins/Experimental/NNERuntimeRDG/Source/NNERuntimeRDG/Private/Tests/NNERuntimeRDGUnitTestOpUnary.cpp
@@ -7,6 +7,7 @@
 #include "NNERuntimeRDGUnitTestHelper.h"
 #include "NNERuntimeRDGTensor.h"
 #include "NNETypes.h"
+#include <limits.h>
 
 #if WITH_DEV_AUTOMATION_TESTS
 #include "Misc/AutomationTest.h"
@@ -206,7 +207,7 @@ namespace UE::NNERuntimeRDG::Private::Test::UnaryOp
 	IMPLEMENT_NNE_SHAPEINFERENCEHELPER_UNIT_AUTOMATION_TEST(FElementWiseUnaryCPUHelperIsInf, "System.Engine.MachineLearning.NNE.RDG.UnitTest.UnaryHelper.IsInf");
 	bool FElementWiseUnaryCPUHelperIsInf::RunTest(const FString& Parameter)
 	{
-		FTensor XC2x1 = MakeConstTensor(TEXT("XC2x1"), { 2,1 }, { -0.5f, INFINITY });
+		FTensor XC2x1 = MakeConstTensor(TEXT("XC2x1"), { 2,1 }, { -0.5f, std::numeric_limits<float>::infinity() });
 		FTensor Y = MakeTensor(TEXT("Y"), { 2,1 });
 		CPUHelper::ElementWiseUnary::Apply(EElementWiseUnaryOperatorType::IsInf, XC2x1, 1.0f, 2.0f, 3.0f, Y);
 		UTEST_TRUE(TEXT("Y const if input is const"), Y.HasPreparedData());
